GameplayTasks work merged over from dev branch #UE4

Introducing AIResource classes. Updated AITask_Move as a first, test AI task
- also more GameplayTasks v-logging added
- added more unit tests for GameplayTasks

diff --git a/Engine/Source/Runtime/AIModule/Classes/AIController.h b/Engine/Source/Runtime/AIModule/Classes/AIController.h
--- a/Engine/Source/Runtime/AIModule/Classes/AIController.h
+++ b/Engine/Source/Runtime/AIModule/Classes/AIController.h
@@ -349,11 +349,12 @@ public:
 	// IGameplayTaskOwnerInterface
 	//----------------------------------------------------------------------//
 	virtual UGameplayTasksComponent* GetGameplayTasksComponent() override { return CachedGameplayTasksComponent; }
-	virtual void TaskStarted(UGameplayTask& NewTask) {}
-	virtual void TaskEnded(UGameplayTask& Task) {}
+	virtual void OnTaskActivated(UGameplayTask& Task) override {}
+	virtual void OnTaskDeactivated(UGameplayTask& Task) override {}
 	virtual void OnTaskInitialized(UGameplayTask& Task) override {}
 	virtual AActor* GetOwnerActor() const { return const_cast<AAIController*>(this); }
 	virtual AActor* GetAvatarActor() const { return GetPawn(); }
+	virtual uint8 GetDefaultPriority() const { return FGameplayTasks::DefaultPriority - 1; }
 
 	//----------------------------------------------------------------------//
 	// Actions
diff --git a/Engine/Source/Runtime/AIModule/Classes/AIResources.h b/Engine/Source/Runtime/AIModule/Classes/AIResources.h
new file mode 100644
--- /dev/null
+++ b/Engine/Source/Runtime/AIModule/Classes/AIResources.h
@@ -0,0 +1,17 @@
+// Copyright 1998-2015 Epic Games, Inc. All Rights Reserved.
+#pragma once
+
+#include "GameplayTaskResource.h"
+#include "AIResources.generated.h"
+
+UCLASS()
+class AIMODULE_API UAIResource_Movement : public UGameplayTaskResource
+{
+	GENERATED_BODY()
+};
+
+UCLASS()
+class AIMODULE_API UAIResource_Logic : public UGameplayTaskResource
+{
+	GENERATED_BODY()
+};
diff --git a/Engine/Source/Runtime/AIModule/Classes/Tasks/AITask_MoveTo.h b/Engine/Source/Runtime/AIModule/Classes/Tasks/AITask_MoveTo.h
--- a/Engine/Source/Runtime/AIModule/Classes/Tasks/AITask_MoveTo.h
+++ b/Engine/Source/Runtime/AIModule/Classes/Tasks/AITask_MoveTo.h
@@ -33,6 +33,9 @@ protected:
 	virtual void HandleMoveFinished(FAIRequestID RequestID, EPathFollowingResult::Type Result);
 	virtual void Activate() override;
 
+	virtual void Pause() override;
+	virtual void Resume() override;
+
 public:
 	UAITask_MoveTo(const FObjectInitializer& ObjectInitializer);
 
diff --git a/Engine/Source/Runtime/AIModule/Private/AIController.cpp b/Engine/Source/Runtime/AIModule/Private/AIController.cpp
--- a/Engine/Source/Runtime/AIModule/Private/AIController.cpp
+++ b/Engine/Source/Runtime/AIModule/Private/AIController.cpp
@@ -142,15 +142,20 @@ void AAIController::GrabDebugSnapshot(FVisualLogEntry* Snapshot) const
 		PathFollowingComponent->DescribeSelfToVisLog(Snapshot);
 	}
 
-	if (BrainComponent != NULL)
+	if (BrainComponent != nullptr)
 	{
 		BrainComponent->DescribeSelfToVisLog(Snapshot);
 	}
 
-	if (PerceptionComponent != NULL)
+	if (PerceptionComponent != nullptr)
 	{
 		PerceptionComponent->DescribeSelfToVisLog(Snapshot);
 	}
+
+	if (CachedGameplayTasksComponent != nullptr)
+	{
+		CachedGameplayTasksComponent->DescribeSelfToVisLog(Snapshot);
+	}
 }
 #endif // ENABLE_VISUAL_LOG
 
@@ -470,6 +475,11 @@ void AAIController::Possess(APawn* InPawn)
 	}
 	CachedGameplayTasksComponent = GTComp;
 
+	if (GTComp)
+	{
+		REDIRECT_OBJECT_TO_VLOG(GTComp, this);
+	}
+
 	OnPossess(InPawn);
 }
 
diff --git a/Engine/Source/Runtime/AIModule/Private/AIResources.cpp b/Engine/Source/Runtime/AIModule/Private/AIResources.cpp
new file mode 100644
--- /dev/null
+++ b/Engine/Source/Runtime/AIModule/Private/AIResources.cpp
@@ -0,0 +1,2 @@
+// Copyright 1998-2015 Epic Games, Inc. All Rights Reserved.
+#include "AIModulePrivate.h"
\ No newline at end of file
diff --git a/Engine/Source/Runtime/AIModule/Private/Tasks/AITask_MoveTo.cpp b/Engine/Source/Runtime/AIModule/Private/Tasks/AITask_MoveTo.cpp
--- a/Engine/Source/Runtime/AIModule/Private/Tasks/AITask_MoveTo.cpp
+++ b/Engine/Source/Runtime/AIModule/Private/Tasks/AITask_MoveTo.cpp
@@ -42,7 +42,10 @@ UAITask_MoveTo* UAITask_MoveTo::AIMoveTo(AAIController* Controller, FVector Targ
 		//MyTask->MoveReq.SetUsePathfinding(bUsePathfinding);
 		//MyTask->MoveRequest.SetNavigationFilter(FilterClass);
 		//MyTask->MoveRequest.SetCanStrafe(bCanStrafe);
+
+		MyTask->RequireResource(UAIResource_Movement::StaticClass());
 	}
+
 	return MyTask;
 }
 
@@ -50,13 +53,20 @@ void UAITask_MoveTo::HandleMoveFinished(FAIRequestID RequestID, EPathFollowingRe
 {
 	if (RequestID == MoveRequestID)
 	{
-		EndTask();
-		OnMoveFinished.Broadcast(Result);
+		if (Result == EPathFollowingResult::Success)
+		{
+			EndTask();
+			OnMoveFinished.Broadcast(Result);
+		}
+		else
+		{
+			Pause();
+		}
 	}
 	else
 	{
 		// @todo report issue to the owner component
-		UE_VLOG(GetGameplayTasksComponent(), LogGameplayTask, Warning, TEXT("%s got movement-finished-notification but RequestID doesn't match. Possible task leak"), *GetName());
+		UE_VLOG(GetGameplayTasksComponent(), LogGameplayTasks, Warning, TEXT("%s got movement-finished-notification but RequestID doesn't match. Possible task leak"), *GetName());
 	}
 }
 
@@ -90,3 +100,18 @@ void UAITask_MoveTo::Activate()
 		break;
 	}
 }
+
+void UAITask_MoveTo::Pause()
+{
+	if (OwnerController)
+	{
+		OwnerController->GetPathFollowingComponent()->OnMoveFinished.RemoveAll(this);
+	}
+	Super::Pause();
+}
+
+void UAITask_MoveTo::Resume()
+{
+	Activate();
+	Super::Resume();
+}
diff --git a/Engine/Source/Runtime/Engine/Public/VisualLogger/VisualLoggerTypes.h b/Engine/Source/Runtime/Engine/Public/VisualLogger/VisualLoggerTypes.h
--- a/Engine/Source/Runtime/Engine/Public/VisualLogger/VisualLoggerTypes.h
+++ b/Engine/Source/Runtime/Engine/Public/VisualLogger/VisualLoggerTypes.h
@@ -107,6 +107,11 @@ struct ENGINE_API FVisualLogStatusCategory
 	int32 UniqueId;
 	TArray<FVisualLogStatusCategory> Children;
 
+	FVisualLogStatusCategory(const FString& InCategory = TEXT(""))
+		: Category(InCategory)
+	{
+	}
+
 	void Add(const FString& Key, const FString& Value);
 	bool GetDesc(int32 Index, FString& Key, FString& Value) const;
 	void AddChild(const FVisualLogStatusCategory& Child);
diff --git a/Engine/Source/Runtime/GameplayAbilities/Private/Abilities/GameplayAbility.cpp b/Engine/Source/Runtime/GameplayAbilities/Private/Abilities/GameplayAbility.cpp
--- a/Engine/Source/Runtime/GameplayAbilities/Private/Abilities/GameplayAbility.cpp
+++ b/Engine/Source/Runtime/GameplayAbilities/Private/Abilities/GameplayAbility.cpp
@@ -1033,11 +1033,11 @@ void UGameplayAbility::OnTaskInitialized(UGameplayTask& Task)
 	}
 }
 
-void UGameplayAbility::TaskStarted(UGameplayTask& NewTask)
+void UGameplayAbility::OnTaskActivated(UGameplayTask& Task)
 {
-	ABILITY_VLOG(CastChecked<AActor>(GetOuter()), Log, TEXT("Task Started %s"), *NewTask.GetName());
+	ABILITY_VLOG(CastChecked<AActor>(GetOuter()), Log, TEXT("Task Started %s"), *Task.GetName());
 
-	ActiveTasks.Add(&NewTask);
+	ActiveTasks.Add(&Task);
 }
 
 void UGameplayAbility::ConfirmTaskByInstanceName(FName InstanceName, bool bEndTask)
@@ -1110,7 +1110,7 @@ void UGameplayAbility::EndAbilityState(FName OptionalStateNameToEnd)
 	}
 }
 
-void UGameplayAbility::TaskEnded(UGameplayTask& Task)
+void UGameplayAbility::OnTaskDeactivated(UGameplayTask& Task)
 {
 	ABILITY_VLOG(CastChecked<AActor>(GetOuter()), Log, TEXT("Task Ended %s"), *Task.GetName());
 
diff --git a/Engine/Source/Runtime/GameplayAbilities/Public/Abilities/GameplayAbility.h b/Engine/Source/Runtime/GameplayAbilities/Public/Abilities/GameplayAbility.h
--- a/Engine/Source/Runtime/GameplayAbilities/Public/Abilities/GameplayAbility.h
+++ b/Engine/Source/Runtime/GameplayAbilities/Public/Abilities/GameplayAbility.h
@@ -255,10 +255,10 @@ public:
 	/** setup additional properties if given task is an AbilityTask */
 	virtual void OnTaskInitialized(UGameplayTask& Task) override;
 	/** Called by an ability task, originating from this ability, when it starts */
-	virtual void TaskStarted(UGameplayTask& NewTask) override;
+	virtual void OnTaskActivated(UGameplayTask& Task) override;
 
 	/** Called by an ability task, originating from this ability, when it ends */
-	virtual void TaskEnded(UGameplayTask& Task) override;
+	virtual void OnTaskDeactivated(UGameplayTask& Task) override;
 
 	virtual UGameplayTasksComponent* GetGameplayTasksComponent() override;
 	virtual AActor* GetOwnerActor() const override;
diff --git a/Engine/Source/Runtime/GameplayTasks/Classes/GameplayTask.h b/Engine/Source/Runtime/GameplayTasks/Classes/GameplayTask.h
--- a/Engine/Source/Runtime/GameplayTasks/Classes/GameplayTask.h
+++ b/Engine/Source/Runtime/GameplayTasks/Classes/GameplayTask.h
@@ -6,19 +6,31 @@
 #include "GameplayTask.generated.h"
 
 class UGameplayTask;
+class UGameplayTaskResource;
 class UGameplayTasksComponent; 
 
-GAMEPLAYTASKS_API DECLARE_LOG_CATEGORY_EXTERN(LogGameplayTask, Log, All);
+GAMEPLAYTASKS_API DECLARE_LOG_CATEGORY_EXTERN(LogGameplayTasks, Log, All);
 
 UENUM()
 enum class EGameplayTaskState : uint8
 {
 	Uninitialized,
 	AwaitingActivation,
+	Paused,
 	Active,
 	Finished
 };
 
+UENUM()
+enum class ETaskResourceOverlapPolicy : uint8
+{
+	// Pause overlapping same-priority tasks
+	StartOnTop,
+	// Wait for other same-priority tasks to finish
+	StartAtEnd,
+};
+	
+
 UCLASS(Abstract)
 class GAMEPLAYTASKS_API UGameplayTask : public UObject
 {
@@ -37,7 +49,7 @@ public:
 protected:
 	/** Called to trigger the actual task once the delegates have been set up
 	 *	Note that the default implementation does nothing and you don't have to call it */
-	virtual void Activate() {}
+	virtual void Activate();
 
 	/** Initailizes the task with the task owner interface instance but does not actviate until Activate() is called */
 	void InitTask(IGameplayTaskOwnerInterface& InTaskOwner);
@@ -80,7 +92,7 @@ public:
 	}
 
 	template <class T>
-	static T* NewTask(UObject& WorldContextObject, FName InstanceName = FName());
+	inline static T* NewTask(UObject& WorldContextObject, FName InstanceName = FName());
 
 	template <class T>
 	inline static T* NewTask(IGameplayTaskOwnerInterface& TaskOwner, FName InstanceName = FName());
@@ -102,22 +114,69 @@ public:
 	FORCEINLINE bool IsSimulatedTask() const { return (bSimulatedTask != 0); }
 	FORCEINLINE bool IsSimulating() const { return (bIsSimulating != 0); }
 	FORCEINLINE bool IsPausable() const { return (bIsPausable != 0); }
-
+	FORCEINLINE uint8 GetPriority() const { return Priority; }
+	FORCEINLINE bool RequiresPriorityOrResourceManagement() const { return bCaresAboutPriority == true || RequiredResources.IsEmpty() == false; }
+	FORCEINLINE FGameplayResourceSet GetRequiredResources() { return RequiredResources; }
+	
 	FORCEINLINE EGameplayTaskState GetState() const { return TaskState; }
-
+	FORCEINLINE bool IsActive() const { return (TaskState == EGameplayTaskState::Active); }
+	
+	IGameplayTaskOwnerInterface* GetTaskOwner() const { return TaskOwner.IsValid() ? &(*TaskOwner) : nullptr; }
 	UGameplayTasksComponent* GetGameplayTasksComponent() { return TasksComponent.Get(); }
+	bool IsOwnedByTasksComponent() const { return bOwnedByTasksComponent; }
+
+	template <class T>
+	inline void RequireResource()
+	{
+		RequireResource(T::StaticClass());
+	}
+
+	/** Marks this task as requiring specified resource which has a number of consequences,
+	*	like task not being able to run if the resource is already taken.
+	*
+	*	@node: Calling this function makes sense only until the task is being passed over to the GameplayTasksComponent.
+	*	Once that's that resources data is consumed and further changes won't get applied */
+	void RequireResource(TSubclassOf<UGameplayTaskResource> RequiredResource);
+
+	ETaskResourceOverlapPolicy GetResourceOverlapPolicy() const { return ResourceOverlapPolicy; }
 
 protected:
-	/** End and CleanUp the task - may be called by the task itself or by the task owner if the owner is ending. Do NOT call directly! Call EndTask() or TaskOwnerEnded() */
+	/** End and CleanUp the task - may be called by the task itself or by the task owner if the owner is ending. 
+	 *	IMPORTANT! Do NOT call directly! Call EndTask() or TaskOwnerEnded() 
+	 *	IMPORTANT! When overriding this function make sure to call Super::OnDestroy(bOwnerFinished) as the last thing,
+	 *		since the function internally marks the task as "Pending Kill", and this may interfere with internal BP mechanics
+	 */
 	virtual void OnDestroy(bool bOwnerFinished);
 
 	static IGameplayTaskOwnerInterface* ConvertToTaskOwner(UObject& OwnerObject);
 	static IGameplayTaskOwnerInterface* ConvertToTaskOwner(AActor& OwnerActor);
 
+	// protected by design. Not meant to be called outside from GameplayTaskComponent mechanics
+	virtual void Pause();
+	virtual void Resume();
+
+private:
+	friend UGameplayTasksComponent;
+
+	void ActivateInTaskQueue();
+	void PauseInTaskQueue();
+	
+	void PerformActivation();
+	
+protected:
+
 	/** This name allows us to find the task later so that we can end it. */
 	UPROPERTY()
 	FName InstanceName;
 
+	/** This controls how this task will be treaded in relation to other, already running tasks, 
+	 *	provided GameplayTasksComponent is configured to care about priorities (the default behavior)*/
+	uint8 Priority;
+
+	EGameplayTaskState TaskState;
+
+	ETaskResourceOverlapPolicy ResourceOverlapPolicy;
+
 	/** If true, this task will receive TickTask calls from TasksComponent */
 	uint32 bTickingTask : 1;
 
@@ -129,31 +188,43 @@ protected:
 
 	uint32 bIsPausable : 1;
 
+	uint32 bCaresAboutPriority : 1;
+
 	/** this is set to avoid duplicate calls to task's owner and TasksComponent when both are the same object */
 	uint32 bOwnedByTasksComponent : 1;
-
-	EGameplayTaskState TaskState;
+	
+	/** Abstract "resource" IDs this task claims it needs. Note that we do not support dynamic changes to this
+	 *	variable. Once a task gets submitted to GameplayTasksComponent it's settled 
+	 *	@NOTE: if dynamic changes were needed though look for GetRequiredResources occurrences in GameplayTasksComponent
+	 *	and replace all uses of cached values with actual GetRequiredResources calls */
+	FGameplayResourceSet RequiredResources;
 
 	/** Task Owner that created us */
 	TWeakInterfacePtr<IGameplayTaskOwnerInterface> TaskOwner;
 
 	TWeakObjectPtr<UGameplayTasksComponent>	TasksComponent;
+
+#if ENABLE_VISUAL_LOG
+	mutable FString DebugDescription;
+public:
+	const FString& GetDebugDescription() const
+	{
+		if (DebugDescription.IsEmpty())
+		{
+			DebugDescription = GenerateDebugDescription();
+		}
+		return DebugDescription;
+	}
+	virtual FString GenerateDebugDescription() const;
+	FString GetTaskStateName() const;
+#endif // ENABLE_VISUAL_LOG
 };
 
 template <class T>
 T* UGameplayTask::NewTask(UObject& WorldContextObject, FName InstanceName)
 {
 	IGameplayTaskOwnerInterface* TaskOwner = ConvertToTaskOwner(WorldContextObject);
-
-	if (TaskOwner)
-	{
-		T* MyObj = NewObject<T>();
-		MyObj->InitTask(*TaskOwner);
-		MyObj->InstanceName = InstanceName;
-		return MyObj;
-	}
-
-	return nullptr;
+	return (TaskOwner) ? NewTask<T>(*TaskOwner, InstanceName) : nullptr;
 }
 
 template <class T>
@@ -162,6 +233,7 @@ T* UGameplayTask::NewTask(IGameplayTaskOwnerInterface& TaskOwner, FName Instance
 	T* MyObj = NewObject<T>();
 	MyObj->InitTask(TaskOwner);
 	MyObj->InstanceName = InstanceName;
+	MyObj->Priority = TaskOwner.GetDefaultPriority();
 	return MyObj;
 }
 
diff --git a/Engine/Source/Runtime/GameplayTasks/Classes/GameplayTaskOwnerInterface.h b/Engine/Source/Runtime/GameplayTasks/Classes/GameplayTaskOwnerInterface.h
--- a/Engine/Source/Runtime/GameplayTasks/Classes/GameplayTaskOwnerInterface.h
+++ b/Engine/Source/Runtime/GameplayTasks/Classes/GameplayTaskOwnerInterface.h
@@ -7,6 +7,11 @@ class UGameplayTasksComponent;
 class UGameplayTask;
 class AActor;
 
+namespace FGameplayTasks
+{
+	static const uint8 DefaultPriority = 127;
+}
+
 UINTERFACE(MinimalAPI, meta = (CannotImplementInterfaceInBlueprint))
 class UGameplayTaskOwnerInterface : public UInterface
 {
@@ -19,8 +24,11 @@ class GAMEPLAYTASKS_API IGameplayTaskOwnerInterface
 public:
 	virtual void OnTaskInitialized(UGameplayTask& Task);
 	virtual UGameplayTasksComponent* GetGameplayTasksComponent() = 0;
-	virtual void TaskStarted(UGameplayTask& NewTask) = 0;
-	virtual void TaskEnded(UGameplayTask& Task) = 0;
+	/** this gets called both when task starts and when task gets resumed. Check Task.GetStatus() if you want to differenciate */
+	virtual void OnTaskActivated(UGameplayTask& Task) = 0;
+	/** this gets called both when task finished and when task gets paused. Check Task.GetStatus() if you want to differenciate */
+	virtual void OnTaskDeactivated(UGameplayTask& Task) = 0;
 	virtual AActor* GetOwnerActor() const = 0;
 	virtual AActor* GetAvatarActor() const;
+	virtual uint8 GetDefaultPriority() const { return FGameplayTasks::DefaultPriority; }
 };
\ No newline at end of file
diff --git a/Engine/Source/Runtime/GameplayTasks/Classes/GameplayTaskResource.h b/Engine/Source/Runtime/GameplayTasks/Classes/GameplayTaskResource.h
new file mode 100644
--- /dev/null
+++ b/Engine/Source/Runtime/GameplayTasks/Classes/GameplayTaskResource.h
@@ -0,0 +1,46 @@
+// Copyright 1998-2015 Epic Games, Inc. All Rights Reserved.
+#pragma once
+
+#include "GameplayTaskResource.generated.h"
+
+UCLASS(Abstract, config="Game")
+class GAMEPLAYTASKS_API UGameplayTaskResource : public UObject
+{
+	GENERATED_BODY()
+
+protected:
+	/** Overrides AutoResourceID. -1 means auto ID will be applied */
+	UPROPERTY(EditDefaultsOnly, BlueprintReadOnly, Category = "Task Resource", meta = (ClampMin = "-1", ClampMax = "15", UIMin = "-1", UIMax = "15", EditCondition = "bManuallySetID"), config)
+	int32 ManualResourceID;
+
+private:
+	int8 AutoResourceID;
+
+public:
+	UPROPERTY()
+	uint32 bManuallySetID : 1;
+
+public:
+
+	UGameplayTaskResource(const FObjectInitializer& ObjectInitializer);
+
+	uint8 GetResourceID() const
+	{
+		return bManuallySetID || (ManualResourceID != INDEX_NONE) ? ManualResourceID : AutoResourceID;
+	}
+
+	static uint8 GetResourceID(TSubclassOf<UGameplayTaskResource>& RequiredResource)
+	{
+		return RequiredResource->GetDefaultObject<UGameplayTaskResource>()->GetResourceID();
+	}
+
+	virtual void PostInitProperties() override;
+
+protected:
+#if WITH_EDITOR
+	void PostEditChangeProperty(FPropertyChangedEvent& PropertyChangedEvent);
+#endif // WITH_EDITOR
+
+	void UpdateAutoResourceID();
+};
+
diff --git a/Engine/Source/Runtime/GameplayTasks/Classes/GameplayTasksComponent.h b/Engine/Source/Runtime/GameplayTasks/Classes/GameplayTasksComponent.h
--- a/Engine/Source/Runtime/GameplayTasks/Classes/GameplayTasksComponent.h
+++ b/Engine/Source/Runtime/GameplayTasks/Classes/GameplayTasksComponent.h
@@ -3,11 +3,30 @@
 #pragma once
 
 #include "GameplayTaskOwnerInterface.h"
+#include "GameplayTaskTypes.h"
+#include "GameplayTask.h"
 #include "GameplayTasksComponent.generated.h"
 
 class FOutBunch;
 class UActorChannel;
-class UGameplayTask;
+
+enum class EGameplayTaskEvent
+{
+	Add,
+	Remove,
+};
+
+struct FGameplayTaskEventData
+{
+	EGameplayTaskEvent Event;
+	UGameplayTask& RelatedTask;
+
+	FGameplayTaskEventData(EGameplayTaskEvent InEvent, UGameplayTask& InRelatedTask)
+		: Event(InEvent), RelatedTask(InRelatedTask)
+	{
+
+	}
+};
 
 /**
 *	The core ActorComponent for interfacing with the GameplayAbilities System
@@ -22,14 +41,23 @@ protected:
 	UPROPERTY(ReplicatedUsing = OnRep_SimulatedTasks)
 	TArray<UGameplayTask*> SimulatedTasks;
 
-	/** tasks that have finished but have not been destroyed yet. We destroy tasks in a latent 
-	 *	manner mostly to be able to safely send out broadcast notifications after the task has 
-	 *	been marked as "finished" */
 	UPROPERTY()
-	TArray<UGameplayTask*> FinishedTasks;
+	TArray<UGameplayTask*> TaskPriorityQueue;
+
+	/** Transient array of events whose main role is to avoid
+	 *	long chain of recurrent calls if an activated/paused/removed task 
+	 *	wants to push/pause/kill other tasks.
+	 *	Note: this TaskEvents is assumed to be used in a single thread */
+	TArray<FGameplayTaskEventData> TaskEvents;
 
 	/** Array of currently active UAbilityTasks that require ticking */
-	TArray<TWeakObjectPtr<UGameplayTask> >	TickingTasks;
+	TArray<TWeakObjectPtr<UGameplayTask> > TickingTasks;
+
+	/** Indicates what's the highest priority among currently running tasks */
+	uint8 TopActivePriority;
+
+	/** Resources used by currently active tasks */
+	FGameplayResourceSet CurrentlyUsedResources;
 
 public:
 	UGameplayTasksComponent(const FObjectInitializer& ObjectInitializer);
@@ -49,12 +77,35 @@ public:
 	
 	virtual UGameplayTasksComponent* GetGameplayTasksComponent() override { return this; }
 	/** Adds a task (most often an UAbilityTask) to the list of tasks to be ticked */
-	virtual void TaskStarted(UGameplayTask& NewTask) override;
+	virtual void OnTaskActivated(UGameplayTask& Task) override;
 	/** Removes a task (most often an UAbilityTask) task from the list of tasks to be ticked */
-	virtual void TaskEnded(UGameplayTask& Task) override;
+	virtual void OnTaskDeactivated(UGameplayTask& Task) override;
+
 	virtual AActor* GetOwnerActor() const override { return GetOwner(); }
 	virtual AActor* GetAvatarActor() const override;
+		
+	/** processes the task and figures out if it should get triggered instantly or wait
+	 *	based on task's RequiredResources, Priority and ResourceOverlapPolicy */
+	void AddTaskReadyForActivation(UGameplayTask& NewTask);
+
+	void RemoveResourceConsumingTask(UGameplayTask& Task);
+
+	FORCEINLINE FGameplayResourceSet GetCurrentlyUsedResources() const { return CurrentlyUsedResources; }
+
+#if ENABLE_VISUAL_LOG
+	static FString GetTaskStateName(EGameplayTaskState Value);
+	void DescribeSelfToVisLog(struct FVisualLogEntry* Snapshot) const;
+#endif // ENABLE_VISUAL_LOG
 
 protected:
 	void RequestTicking();
+	void ProcessTaskEvents();
+	void UpdateTaskActivationFromIndex(int32 StartingIndex, FGameplayResourceSet ResourcesUsedByActiveUpToIndex, FGameplayResourceSet ResourcesRequiredUpToIndex);
+
+private:
+	/** called when a task gets ended with an external call, meaning not coming from UGameplayTasksComponent mechanics */
+	void OnTaskEnded(UGameplayTask& Task);
+
+	void AddTaskToPriorityQueue(UGameplayTask& NewTask);
+	void RemoveTaskFromPriorityQueue(UGameplayTask& Task);
 };
\ No newline at end of file
diff --git a/Engine/Source/Runtime/GameplayTasks/Private/GameplayTask.cpp b/Engine/Source/Runtime/GameplayTasks/Private/GameplayTask.cpp
--- a/Engine/Source/Runtime/GameplayTasks/Private/GameplayTask.cpp
+++ b/Engine/Source/Runtime/GameplayTasks/Private/GameplayTask.cpp
@@ -2,6 +2,7 @@
 
 #include "GameplayTasksPrivatePCH.h"
 #include "GameplayTasksComponent.h"
+#include "GameplayTaskResource.h"
 #include "GameplayTask.h"
 
 UGameplayTask::UGameplayTask(const FObjectInitializer& ObjectInitializer)
@@ -12,6 +13,7 @@ UGameplayTask::UGameplayTask(const FObjectInitializer& ObjectInitializer)
 	bIsSimulating = false;
 	bOwnedByTasksComponent = false;
 	TaskState = EGameplayTaskState::Uninitialized;
+	ResourceOverlapPolicy = ETaskResourceOverlapPolicy::StartOnTop;
 	SetFlags(RF_StrongRefOnFrame);
 }
 
@@ -43,21 +45,16 @@ IGameplayTaskOwnerInterface* UGameplayTask::ConvertToTaskOwner(AActor& OwnerActo
 
 void UGameplayTask::ReadyForActivation()
 {
-	// default case, we instantly activate the task
-	TaskState = EGameplayTaskState::Active;
-
-	if (TaskOwner.IsValid())
+	if (TasksComponent.IsValid())
 	{
-		TaskOwner->TaskStarted(*this);
-
-		// If this task requires ticking, register it with TasksComponent
-		if (bOwnedByTasksComponent == false && TasksComponent.IsValid())
+		if (RequiresPriorityOrResourceManagement() == false)
 		{
-			TasksComponent->TaskStarted(*this);
+			PerformActivation();
+		}
+		else
+		{
+			TasksComponent->AddTaskReadyForActivation(*this);
 		}
-		// else: it's not supposed to work this way!
-
-		Activate();
 	}
 	else
 	{
@@ -128,6 +125,10 @@ AActor* UGameplayTask::GetAvatarActor() const
 
 void UGameplayTask::TaskOwnerEnded()
 {
+	UE_VLOG(GetGameplayTasksComponent(), LogGameplayTasks, Verbose
+		, TEXT("%s TaskOwnerEnded called, current State: %s")
+		, *GetName(), *GetTaskStateName());
+
 	if (TaskState != EGameplayTaskState::Finished && !IsPendingKill())
 	{
 		OnDestroy(true);
@@ -136,6 +137,10 @@ void UGameplayTask::TaskOwnerEnded()
 
 void UGameplayTask::EndTask()
 {
+	UE_VLOG(GetGameplayTasksComponent(), LogGameplayTasks, Verbose
+		, TEXT("%s EndTask called, current State: %s")
+		, *GetName(), *GetTaskStateName());
+
 	if (TaskState != EGameplayTaskState::Finished && !IsPendingKill())
 	{
 		OnDestroy(false);
@@ -144,6 +149,10 @@ void UGameplayTask::EndTask()
 
 void UGameplayTask::ExternalConfirm(bool bEndTask)
 {
+	UE_VLOG(GetGameplayTasksComponent(), LogGameplayTasks, Verbose
+		, TEXT("%s ExternalConfirm called, bEndTask = %s, State : %s")
+		, *GetName(), bEndTask ? TEXT("TRUE") : TEXT("FALSE"), *GetTaskStateName());
+
 	if (bEndTask)
 	{
 		EndTask();
@@ -152,6 +161,10 @@ void UGameplayTask::ExternalConfirm(bool bEndTask)
 
 void UGameplayTask::ExternalCancel()
 {
+	UE_VLOG(GetGameplayTasksComponent(), LogGameplayTasks, Verbose
+		, TEXT("%s ExternalCancel called, current State: %s")
+		, *GetName(), *GetTaskStateName());
+
 	EndTask();
 }
 
@@ -161,30 +174,165 @@ void UGameplayTask::OnDestroy(bool bOwnerFinished)
 
 	TaskState = EGameplayTaskState::Finished;
 
-	bool bMarkAsPendingKill = true;
-	// If this task required ticking, unregister it with TasksComponent
+	// First of all notify the TaskComponent
 	if (TasksComponent.IsValid())
 	{
-		TasksComponent->TaskEnded(*this);
-		// we need to mark task as pending kill only if there's no component that can 
-		// do that for us, latently
-		bMarkAsPendingKill = false;
+		TasksComponent->OnTaskDeactivated(*this);
 	}
 
-	// Remove ourselves from the owner's task list, if the ability isn't ending
+	// Remove ourselves from the owner's task list, if the owner isn't ending
 	if (bOwnedByTasksComponent == false && bOwnerFinished == false && TaskOwner.IsValid() == true)
 	{
-		TaskOwner->TaskEnded(*this);
+		TaskOwner->OnTaskDeactivated(*this);
 	}
 
-	if (bMarkAsPendingKill)
-	{
-		// this usually means the driven actor and/or the component are already gone
-		MarkPendingKill();
-	}
+	MarkPendingKill();
 }
 
 FString UGameplayTask::GetDebugString() const
 {
 	return FString::Printf(TEXT("Generic %s"), *GetName());
 }
+
+void UGameplayTask::RequireResource(TSubclassOf<UGameplayTaskResource> RequiredResource)
+{
+	check(RequiredResource);
+	const uint8 ResourceID = UGameplayTaskResource::GetResourceID(RequiredResource);
+	RequiredResources.AddID(ResourceID);	
+}
+
+void UGameplayTask::PerformActivation()
+{
+	TaskState = EGameplayTaskState::Active;
+
+	Activate();
+
+	TasksComponent->OnTaskActivated(*this);
+
+	if (bOwnedByTasksComponent == false && TaskOwner.IsValid())
+	{
+		TaskOwner->OnTaskActivated(*this);
+	}
+}
+
+void UGameplayTask::Activate()
+{
+	UE_VLOG(GetGameplayTasksComponent(), LogGameplayTasks, Verbose
+		, TEXT("%s Activate called, current State: %s")
+		, *GetName(), *GetTaskStateName());
+}
+
+void UGameplayTask::Pause()
+{
+	UE_VLOG(GetGameplayTasksComponent(), LogGameplayTasks, Verbose
+		, TEXT("%s Pause called, current State: %s")
+		, *GetName(), *GetTaskStateName());
+
+	TaskState = EGameplayTaskState::Paused;
+
+	TasksComponent->OnTaskDeactivated(*this);
+
+	if (bOwnedByTasksComponent == false && TaskOwner.IsValid())
+	{
+		TaskOwner->OnTaskDeactivated(*this);
+	}
+}
+
+void UGameplayTask::Resume()
+{
+	UE_VLOG(GetGameplayTasksComponent(), LogGameplayTasks, Verbose
+		, TEXT("%s Resume called, current State: %s")
+		, *GetName(), *GetTaskStateName());
+
+	TaskState = EGameplayTaskState::Active;
+
+	TasksComponent->OnTaskActivated(*this);
+
+	if (bOwnedByTasksComponent == false && TaskOwner.IsValid())
+	{
+		TaskOwner->OnTaskActivated(*this);
+	}
+}
+
+//----------------------------------------------------------------------//
+// GameplayTasksComponent-related functions
+//----------------------------------------------------------------------//
+void UGameplayTask::ActivateInTaskQueue()
+{
+	switch(TaskState)
+	{
+	case EGameplayTaskState::Uninitialized:
+		UE_VLOG(GetGameplayTasksComponent(), LogGameplayTasks, Error
+			, TEXT("UGameplayTask::ActivateInTaskQueue Task %s passed for activation withouth having InitTask called on it!")
+			, *GetName());
+		break;
+	case EGameplayTaskState::AwaitingActivation:
+		PerformActivation();
+		break;
+	case EGameplayTaskState::Paused:
+		// resume
+		Resume();
+		break;
+	case EGameplayTaskState::Active:
+		// nothing to do here
+		break;
+	case EGameplayTaskState::Finished:
+		// If a task has finished, and it's being revived let's just treat the same as AwaitingActivation
+		PerformActivation();
+		break;
+	default:
+		checkNoEntry(); // looks like unhandled value! Probably a new enum entry has been added
+		break;
+	}
+}
+
+void UGameplayTask::PauseInTaskQueue()
+{
+	switch (TaskState)
+	{
+	case EGameplayTaskState::Uninitialized:
+		UE_VLOG(GetGameplayTasksComponent(), LogGameplayTasks, Error
+			, TEXT("UGameplayTask::PauseInTaskQueue Task %s passed for pausing withouth having InitTask called on it!")
+			, *GetName());
+		break;
+	case EGameplayTaskState::AwaitingActivation:
+		// nothing to do here. Don't change the state to indicate this task has never been run before
+		break;
+	case EGameplayTaskState::Paused:
+		// nothing to do here. Already paused
+		break;
+	case EGameplayTaskState::Active:
+		// pause!
+		Pause();
+		break;
+	case EGameplayTaskState::Finished:
+		// nothing to do here. But sounds odd, so let's log this, just in case
+		UE_VLOG(GetGameplayTasksComponent(), LogGameplayTasks, Log
+			, TEXT("UGameplayTask::PauseInTaskQueue Task %s being pause while already marked as Finished")
+			, *GetName());
+		break;
+	default:
+		checkNoEntry(); // looks like unhandled value! Probably a new enum entry has been added
+		break;
+	}
+}
+
+#if ENABLE_VISUAL_LOG
+//----------------------------------------------------------------------//
+// debug
+//----------------------------------------------------------------------//
+FString UGameplayTask::GenerateDebugDescription() const
+{
+	return RequiresPriorityOrResourceManagement() == false ? GetName()
+		: FString::Printf(TEXT("%s: P:%d %s")
+			, *GetName(), int32(Priority), *RequiredResources.GetDebugDescription());
+}
+
+FString UGameplayTask::GetTaskStateName() const
+{
+	static const UEnum* Enum = FindObject<UEnum>(ANY_PACKAGE, TEXT("EGameplayTaskState"));
+	check(Enum);
+	return Enum->GetEnumName(int32(TaskState));
+}
+
+#endif // ENABLE_VISUAL_LOG
\ No newline at end of file
diff --git a/Engine/Source/Runtime/GameplayTasks/Private/GameplayTaskResource.cpp b/Engine/Source/Runtime/GameplayTasks/Private/GameplayTaskResource.cpp
new file mode 100644
--- /dev/null
+++ b/Engine/Source/Runtime/GameplayTasks/Private/GameplayTaskResource.cpp
@@ -0,0 +1,67 @@
+// Copyright 1998-2015 Epic Games, Inc. All Rights Reserved.
+
+#include "GameplayTasksPrivatePCH.h"
+#include "GameplayTaskResource.h"
+
+UGameplayTaskResource::UGameplayTaskResource(const FObjectInitializer& ObjectInitializer)
+	: Super(ObjectInitializer)
+{
+	bManuallySetID = false;
+	ManualResourceID = INDEX_NONE;
+	AutoResourceID = INDEX_NONE;
+}
+
+void UGameplayTaskResource::PostInitProperties()
+{
+	Super::PostInitProperties();
+
+	if (HasAnyFlags(RF_ClassDefaultObject) && (GetClass()->HasAnyClassFlags(CLASS_Abstract) == false)
+#if WITH_HOT_RELOAD
+		&& !GIsHotReload
+#endif // WITH_HOT_RELOAD
+		)
+	{
+		if (bManuallySetID == false || ManualResourceID == INDEX_NONE)
+		{
+			UpdateAutoResourceID();
+		}
+	}
+}
+
+void UGameplayTaskResource::UpdateAutoResourceID()
+{
+	static uint16 NextAutoResID = 0;
+
+	if (AutoResourceID == INDEX_NONE)
+	{
+		AutoResourceID = NextAutoResID++;
+		if (AutoResourceID >= FGameplayResourceSet::MaxResources)
+		{
+			UE_LOG(LogGameplayTasks, Error, TEXT("AutoResourceID out of bounds (probably too much GameplayTaskResource classes, consider manually assigning values if you can split all classes into non-overlapping sets"));
+		}
+	}
+}
+
+#if WITH_EDITOR
+void UGameplayTaskResource::PostEditChangeProperty(FPropertyChangedEvent& PropertyChangedEvent)
+{
+	static const FName NAME_ManuallySetID = GET_MEMBER_NAME_CHECKED(UGameplayTaskResource, bManuallySetID);
+
+	Super::PostEditChangeProperty(PropertyChangedEvent);
+
+	if (PropertyChangedEvent.Property != NULL)
+	{
+		FName PropName = PropertyChangedEvent.Property->GetFName();
+
+		if (PropName == NAME_ManuallySetID)
+		{
+			if (!bManuallySetID)
+			{
+				ManualResourceID = INDEX_NONE;
+				// if we don't have ManualResourceID make sure AutoResourceID is valid
+				UpdateAutoResourceID();
+			}
+		}
+	}
+}
+#endif // WITH_EDITOR
\ No newline at end of file
diff --git a/Engine/Source/Runtime/GameplayTasks/Private/GameplayTasksComponent.cpp b/Engine/Source/Runtime/GameplayTasks/Private/GameplayTasksComponent.cpp
--- a/Engine/Source/Runtime/GameplayTasks/Private/GameplayTasksComponent.cpp
+++ b/Engine/Source/Runtime/GameplayTasks/Private/GameplayTasksComponent.cpp
@@ -3,7 +3,6 @@
 #include "GameplayTasksPrivatePCH.h"
 #include "GameplayTasksComponent.h"
 #include "GameplayTask.h"
-#include "VisualLogger.h"
 
 #define LOCTEXT_NAMESPACE "GameplayTasksComponent"
 
@@ -15,30 +14,30 @@ UGameplayTasksComponent::UGameplayTasksComponent(const FObjectInitializer& Objec
 	PrimaryComponentTick.bCanEverTick = true;
 
 	bReplicates = true;
+	TopActivePriority = 0;
 }
 	
-void UGameplayTasksComponent::TaskStarted(UGameplayTask& NewTask)
+void UGameplayTasksComponent::OnTaskActivated(UGameplayTask& Task)
 {
-	if (NewTask.IsTickingTask())
+	if (Task.IsTickingTask())
 	{
-		check(TickingTasks.Contains(&NewTask) == false);
-		TickingTasks.Add(&NewTask);
+		check(TickingTasks.Contains(&Task) == false);
+		TickingTasks.Add(&Task);
 
 		// If this is our first ticking task, set this component as active so it begins ticking
 		if (TickingTasks.Num() == 1)
 		{
 			UpdateShouldTick();
 		}
-
 	}
-	if (NewTask.IsSimulatedTask())
+	if (Task.IsSimulatedTask())
 	{
-		check(SimulatedTasks.Contains(&NewTask) == false);
-		SimulatedTasks.Add(&NewTask);
+		check(SimulatedTasks.Contains(&Task) == false);
+		SimulatedTasks.Add(&Task);
 	}
 }
 
-void UGameplayTasksComponent::TaskEnded(UGameplayTask& Task)
+void UGameplayTasksComponent::OnTaskDeactivated(UGameplayTask& Task)
 {
 	if (Task.IsTickingTask())
 	{
@@ -51,8 +50,20 @@ void UGameplayTasksComponent::TaskEnded(UGameplayTask& Task)
 		SimulatedTasks.RemoveSingleSwap(&Task);
 	}
 
-	FinishedTasks.Add(&Task);
-	RequestTicking();
+	UpdateShouldTick();
+
+	// Resource-using task
+	if (Task.RequiresPriorityOrResourceManagement() == true && Task.GetState() == EGameplayTaskState::Finished)
+	{
+		OnTaskEnded(Task);
+	}
+}
+
+void UGameplayTasksComponent::OnTaskEnded(UGameplayTask& Task)
+{
+	ensure(Task.RequiresPriorityOrResourceManagement() == true);
+
+	RemoveTaskFromPriorityQueue(Task);
 }
 
 void UGameplayTasksComponent::GetLifetimeReplicatedProps(TArray< FLifetimeProperty >& OutLifetimeProps) const
@@ -103,14 +114,6 @@ void UGameplayTasksComponent::TickComponent(float DeltaTime, enum ELevelTick Tic
 
 	Super::TickComponent(DeltaTime, TickType, ThisTickFunction);
 
-	// clean up the tasks that have already finished
-	for (auto FinishedTask : FinishedTasks)
-	{
-		ensure(FinishedTask->IsPendingKill() == false);
-		FinishedTask->MarkPendingKill();
-	}
-	FinishedTasks.Reset();
-
 	// Because we have no control over what a task may do when it ticks, we must be careful.
 	// Ticking a task may kill the task right here. It could also potentially kill another task
 	// which was waiting on the original task to do something. Since when a tasks is killed, it removes
@@ -154,7 +157,7 @@ void UGameplayTasksComponent::TickComponent(float DeltaTime, enum ELevelTick Tic
 
 bool UGameplayTasksComponent::GetShouldTick() const
 {
-	return TickingTasks.Num() > 0 || FinishedTasks.Num() > 0;
+	return TickingTasks.Num() > 0;
 }
 
 void UGameplayTasksComponent::RequestTicking()
@@ -179,4 +182,280 @@ AActor* UGameplayTasksComponent::GetAvatarActor() const
 	return GetOwner();
 }
 
+//----------------------------------------------------------------------//
+// Priority and resources handling
+//----------------------------------------------------------------------//
+void UGameplayTasksComponent::AddTaskReadyForActivation(UGameplayTask& NewTask)
+{
+	UE_VLOG(this, LogGameplayTasks, Log, TEXT("RemoveResourceConsumingTask %s"), *NewTask.GetName());
+
+	ensure(NewTask.RequiresPriorityOrResourceManagement() == true);
+	
+	TaskEvents.Add(FGameplayTaskEventData(EGameplayTaskEvent::Add, NewTask));
+	// trigger the actual processing only if it was the first event added to the list
+	if (TaskEvents.Num() == 1)
+	{
+		ProcessTaskEvents();
+	}
+}
+
+void UGameplayTasksComponent::RemoveResourceConsumingTask(UGameplayTask& Task)
+{
+	UE_VLOG(this, LogGameplayTasks, Log, TEXT("RemoveResourceConsumingTask %s"), *Task.GetName());
+
+	TaskEvents.Add(FGameplayTaskEventData(EGameplayTaskEvent::Remove, Task));
+	// trigger the actual processing only if it was the first event added to the list
+	if (TaskEvents.Num() == 1)
+	{
+		ProcessTaskEvents();
+	}
+}
+
+void UGameplayTasksComponent::ProcessTaskEvents()
+{
+	static const int32 ErronousIterationsLimit = 100;
+
+	// note that this function allows called functions to add new elements to 
+	// TaskEvents array that the main loop is iterating over. It's a feature
+	for (int32 EventIndex = 0; EventIndex < TaskEvents.Num(); ++EventIndex)
+	{
+		if (TaskEvents[EventIndex].RelatedTask.IsPendingKill())
+		{
+			// we should ignore it, but just in case run the removal code.
+			RemoveTaskFromPriorityQueue(TaskEvents[EventIndex].RelatedTask);
+			continue;
+		}
+
+		switch (TaskEvents[EventIndex].Event)
+		{
+		case EGameplayTaskEvent::Add:
+			AddTaskToPriorityQueue(TaskEvents[EventIndex].RelatedTask);
+			break;
+		case EGameplayTaskEvent::Remove:
+			RemoveTaskFromPriorityQueue(TaskEvents[EventIndex].RelatedTask);
+			break;
+		default:
+			checkNoEntry();
+			break;
+		}
+
+		if (EventIndex >= ErronousIterationsLimit)
+		{
+			UE_VLOG(this, LogGameplayTasks, Error, TEXT("UGameplayTasksComponent::ProcessTaskEvents has exceeded warning-level of iterations. Check your GameplayTasks for logic loops!"));
+		}
+	}
+
+	TaskEvents.Reset();
+}
+
+void UGameplayTasksComponent::AddTaskToPriorityQueue(UGameplayTask& NewTask)
+{
+	// The generic idea is as follows:
+	// 1. Find insertion/removal point X
+	//		While looking for it add up all ResourcesUsedByActiveUpToX and ResourcesRequiredUpToX
+	//		ResourcesUsedByActiveUpToX - resources required by ACTIVE tasks
+	//		ResourcesRequiredUpToX - resources required by both active and inactive tasks on the way to X
+	// 2. Insert Task at X
+	// 3. Starting from X proceed down the queue
+	//	a. If ConsideredTask.Resources overlaps with ResourcesRequiredUpToX then PAUSE it
+	//	b. Else, ACTIVATE ConsideredTask and add ConsideredTask.Resources to ResourcesUsedByActiveUpToX
+	//	c. Add ConsideredTask.Resources to ResourcesRequiredUpToX
+	// 4. Set this->CurrentlyUsedResources to ResourcesUsedByActiveUpToX
+	//
+	// Points 3-4 are implemented as a separate function, UpdateTaskActivationFromIndex,
+	// since it's a common code between adding and removing tasks
+	
+	const bool bStartOnTopOfSamePriority = NewTask.GetResourceOverlapPolicy() == ETaskResourceOverlapPolicy::StartOnTop;
+
+	FGameplayResourceSet ResourcesUsedByActiveUpToX;
+	FGameplayResourceSet ResourcesRequiredUpToX;
+	int32 InsertionPoint = INDEX_NONE;
+	
+	if (TaskPriorityQueue.Num() > 0)
+	{
+		for (int32 TaskIndex = 0; TaskIndex < TaskPriorityQueue.Num(); ++TaskIndex)
+		{
+			ensure(TaskPriorityQueue[TaskIndex]);
+			if (TaskPriorityQueue[TaskIndex] == nullptr)
+			{
+				continue;
+			}
+
+			if ((bStartOnTopOfSamePriority && TaskPriorityQueue[TaskIndex]->GetPriority() <= NewTask.GetPriority())
+				|| (!bStartOnTopOfSamePriority && TaskPriorityQueue[TaskIndex]->GetPriority() < NewTask.GetPriority()))
+			{
+				TaskPriorityQueue.Insert(&NewTask, TaskIndex);
+				InsertionPoint = TaskIndex;
+				break;
+			}
+
+			const FGameplayResourceSet TasksResources = TaskPriorityQueue[TaskIndex]->GetRequiredResources();
+			if (TaskPriorityQueue[TaskIndex]->IsActive())
+			{
+				ResourcesUsedByActiveUpToX.AddSet(TasksResources);
+			}
+			ResourcesRequiredUpToX.AddSet(TasksResources);
+		}
+	}
+	
+	if (InsertionPoint == INDEX_NONE)
+	{
+		TaskPriorityQueue.Add(&NewTask);
+		InsertionPoint = TaskPriorityQueue.Num() - 1;
+	}
+
+	UpdateTaskActivationFromIndex(InsertionPoint, ResourcesUsedByActiveUpToX, ResourcesRequiredUpToX);
+}
+
+void UGameplayTasksComponent::RemoveTaskFromPriorityQueue(UGameplayTask& Task)
+{	
+	if (TaskPriorityQueue.Num() > 1)
+	{
+		const int32 RemovedTaskIndex = TaskPriorityQueue.Find(&Task);
+		ensure(RemovedTaskIndex != INDEX_NONE);
+
+		// sum up resources up to TaskIndex
+		FGameplayResourceSet ResourcesUsedByActiveUpToX;
+		FGameplayResourceSet ResourcesRequiredUpToX;
+		for (int32 Index = 0; Index < RemovedTaskIndex; ++Index)
+		{
+			ensure(TaskPriorityQueue[Index]);
+			if (TaskPriorityQueue[Index] == nullptr)
+			{
+				continue;
+			}
+			
+			const FGameplayResourceSet TasksResources = TaskPriorityQueue[Index]->GetRequiredResources();
+			if (TaskPriorityQueue[Index]->IsActive())
+			{
+				ResourcesUsedByActiveUpToX.AddSet(TasksResources);
+			}
+			ResourcesRequiredUpToX.AddSet(TasksResources);
+		}
+
+		// don't forget to actually remove the task from the queue
+		TaskPriorityQueue.RemoveAt(RemovedTaskIndex, 1, /*bAllowShrinking=*/false);
+
+		// if it wasn't the last item then proceed as usual
+		if (RemovedTaskIndex < TaskPriorityQueue.Num())
+		{
+			UpdateTaskActivationFromIndex(RemovedTaskIndex, ResourcesUsedByActiveUpToX, ResourcesRequiredUpToX);
+		}
+		else
+		{
+			// no need to do extra processing. This was the last task, so 
+			// ResourcesUsedByActiveUpToX is the CurrentlyUsedResources
+			CurrentlyUsedResources = ResourcesUsedByActiveUpToX;
+		}
+	}
+	else
+	{
+		TaskPriorityQueue.Pop(/*bAllowShrinking=*/false);
+		CurrentlyUsedResources.Clear();
+	}
+}
+
+void UGameplayTasksComponent::UpdateTaskActivationFromIndex(int32 StartingIndex, FGameplayResourceSet ResourcesUsedByActiveUpToIndex, FGameplayResourceSet ResourcesRequiredUpToIndex)
+{
+	check(TaskPriorityQueue.IsValidIndex(StartingIndex));
+
+	TArray<UGameplayTask*> TaskPriorityQueueCopy = TaskPriorityQueue;
+	for (int32 TaskIndex = StartingIndex; TaskIndex < TaskPriorityQueueCopy.Num(); ++TaskIndex)
+	{
+		ensure(TaskPriorityQueueCopy[TaskIndex]);
+		if (TaskPriorityQueueCopy[TaskIndex] == nullptr)
+		{
+			continue;
+		}
+
+		const FGameplayResourceSet TasksResources = TaskPriorityQueueCopy[TaskIndex]->GetRequiredResources();
+		if (TasksResources.GetOverlap(ResourcesUsedByActiveUpToIndex).IsEmpty())
+		{
+			TaskPriorityQueueCopy[TaskIndex]->ActivateInTaskQueue();
+			ResourcesUsedByActiveUpToIndex.AddSet(TasksResources);
+		}
+		else
+		{
+			TaskPriorityQueueCopy[TaskIndex]->PauseInTaskQueue();
+		}
+		ResourcesRequiredUpToIndex.AddSet(TasksResources);
+	}
+
+	CurrentlyUsedResources = ResourcesUsedByActiveUpToIndex;
+}
+
+//----------------------------------------------------------------------//
+// debugging
+//----------------------------------------------------------------------//
+#if ENABLE_VISUAL_LOG
+void UGameplayTasksComponent::DescribeSelfToVisLog(FVisualLogEntry* Snapshot) const
+{
+	static const FString CategoryName = TEXT("GameplayTasks");
+	static const FString TickingTasksName = TEXT("Ticking tasks");
+	static const FString PriorityQueueName = TEXT("Priority Queue");
+
+	if (IsPendingKill())
+	{
+		return;
+	}
+
+	FVisualLogStatusCategory StatusCategory(CategoryName);
+
+	FString TasksDescription;
+	for (auto& Task : TickingTasks)
+	{
+		if (Task.IsValid())
+		{
+			TasksDescription += FString::Printf(TEXT("%s %s\n"), *GetTaskStateName(Task->GetState()), *Task->GetDebugDescription());
+		}
+		else
+		{
+			TasksDescription += TEXT("NULL\n");
+		}
+	}
+	StatusCategory.Add(TickingTasksName, TasksDescription);
+
+	TasksDescription.Reset();
+	for (auto Task : TaskPriorityQueue)
+	{
+		if (Task != nullptr)
+		{
+			TasksDescription += FString::Printf(TEXT("%s %s\n"), *GetTaskStateName(Task->GetState()), *Task->GetDebugDescription());
+		}
+		else
+		{
+			TasksDescription += TEXT("NULL\n");
+		}
+	}
+	StatusCategory.Add(PriorityQueueName, TasksDescription);
+
+	Snapshot->Status.Add(StatusCategory);
+}
+
+FString UGameplayTasksComponent::GetTaskStateName(EGameplayTaskState Value)
+{
+	static const UEnum* Enum = FindObject<UEnum>(ANY_PACKAGE, TEXT("EGameplayTaskState"));
+	check(Enum);
+	return Enum->GetEnumName(int32(Value));
+}
+#endif // ENABLE_VISUAL_LOG
+
+//----------------------------------------------------------------------//
+// FGameplayResourceSet
+//----------------------------------------------------------------------//
+FString FGameplayResourceSet::GetDebugDescription() const
+{
+	static const int32 FlagsCount = sizeof(FFlagContainer)* 8;
+	TCHAR Description[FlagsCount + 1];
+	FFlagContainer FlagsCopy = Flags;
+	int32 FlagIndex = 0;
+	for (; FlagIndex < FlagsCount && FlagsCopy != 0; ++FlagIndex)
+	{
+		Description[FlagIndex] = (FlagsCopy & (1 << FlagIndex)) ? TCHAR('1') : TCHAR('0');
+		FlagsCopy &= ~(1 << FlagIndex);
+	}
+	Description[FlagIndex] = TCHAR('\0');
+	return FString(Description);
+}
+
 #undef LOCTEXT_NAMESPACE
\ No newline at end of file
diff --git a/Engine/Source/Runtime/GameplayTasks/Private/GameplayTasksModule.cpp b/Engine/Source/Runtime/GameplayTasks/Private/GameplayTasksModule.cpp
--- a/Engine/Source/Runtime/GameplayTasks/Private/GameplayTasksModule.cpp
+++ b/Engine/Source/Runtime/GameplayTasks/Private/GameplayTasksModule.cpp
@@ -22,5 +22,5 @@ public:
 //////////////////////////////////////////////////////////////////////////
 
 IMPLEMENT_MODULE(FGameplayTasksModule, GameplayTasks);
-DEFINE_LOG_CATEGORY(LogGameplayTask);
+DEFINE_LOG_CATEGORY(LogGameplayTasks);
 DEFINE_STAT(STAT_TickGameplayTasks);
diff --git a/Engine/Source/Runtime/GameplayTasks/Private/GameplayTasksPrivatePCH.h b/Engine/Source/Runtime/GameplayTasks/Private/GameplayTasksPrivatePCH.h
--- a/Engine/Source/Runtime/GameplayTasks/Private/GameplayTasksPrivatePCH.h
+++ b/Engine/Source/Runtime/GameplayTasks/Private/GameplayTasksPrivatePCH.h
@@ -11,9 +11,12 @@
 
 #include "Net/UnrealNetwork.h"
 #include "Engine/ActorChannel.h"
+#include "VisualLogger.h"
 
 #include "GameplayTasksComponent.h"
 #include "GameplayTask.h"
 
 DECLARE_STATS_GROUP(TEXT("GameplayTasks"), STATGROUP_GameplayTasks, STATCAT_Advanced);
 DECLARE_CYCLE_STAT_EXTERN(TEXT("TickGameplayTasks"), STAT_TickGameplayTasks, STATGROUP_GameplayTasks, );
+
+PRAGMA_DISABLE_OPTIMIZATION
\ No newline at end of file
diff --git a/Engine/Source/Runtime/GameplayTasks/Public/GameplayTaskTypes.h b/Engine/Source/Runtime/GameplayTasks/Public/GameplayTaskTypes.h
--- a/Engine/Source/Runtime/GameplayTasks/Public/GameplayTaskTypes.h
+++ b/Engine/Source/Runtime/GameplayTasks/Public/GameplayTaskTypes.h
@@ -48,4 +48,91 @@ struct TWeakInterfacePtr
 private:
 	TWeakObjectPtr<UObject> ObjectInstance;
 	TInterface* InterfaceInstance;
+};
+
+struct GAMEPLAYTASKS_API FGameplayResourceSet
+{
+	typedef uint16 FFlagContainer;
+	typedef uint8 FResourceID;
+
+	enum
+	{
+		MaxResources = sizeof(FFlagContainer)* 8
+	};
+
+private:
+	FFlagContainer Flags;
+
+public:
+	/** Mind that this constructor takes _flags_ not individual IDs */
+	explicit FGameplayResourceSet(FFlagContainer InFlags = 0) : Flags(InFlags)
+	{}
+	
+	FFlagContainer GetFlags() const 
+	{ 
+		return Flags; 
+	}
+	bool IsEmpty() const
+	{
+		return Flags == 0;
+	}
+	FGameplayResourceSet& AddID(uint8 ResourceID)
+	{
+		ensure(ResourceID < MaxResources);
+		Flags |= (1 << ResourceID);
+		return *this;
+	}
+	FGameplayResourceSet& RemoveID(uint8 ResourceID)
+	{
+		ensure(ResourceID < MaxResources);
+		Flags &= ~(1 << ResourceID);
+		return *this;
+	}
+	bool HasID(uint8 ResourceID) const
+	{
+		ensure(ResourceID < MaxResources);
+		return (Flags & (1 << ResourceID)) != 0;
+	}
+	FGameplayResourceSet& AddSet(FGameplayResourceSet Other)
+	{
+		Flags |= Other.Flags;
+		return *this;
+	}
+	void Clear()
+	{
+		Flags = FFlagContainer(0);
+	}
+	bool HasAllIDs(FGameplayResourceSet Other) const
+	{
+		return (Flags & Other.Flags) == Flags;
+	}
+	bool HasAnyID(FGameplayResourceSet Other) const
+	{
+		return (Flags & Other.Flags) != 0;
+	}
+	FGameplayResourceSet GetOverlap(FGameplayResourceSet Other) const
+	{
+		return FGameplayResourceSet(Flags & Other.Flags);
+	}
+	FGameplayResourceSet GetDifference(FGameplayResourceSet Other) const
+	{
+		return FGameplayResourceSet(Flags & ~(Flags & Other.Flags));
+	}
+
+	bool operator==(const FGameplayResourceSet& Other) const
+	{
+		return Flags == Other.Flags;
+	}
+
+	static FGameplayResourceSet AllResources()
+	{
+		return FGameplayResourceSet(FFlagContainer(-1));
+	}
+
+	static FGameplayResourceSet NoResources()
+	{
+		return FGameplayResourceSet(FFlagContainer(0));
+	}
+
+	FString GetDebugDescription() const;
 };
\ No newline at end of file
