Renamed GameplayTask_SpawnActor's factory function to distinguish it from SpawnActorFromClass k2 node #UE4

diff --git a/Engine/Source/Runtime/GameplayTasks/Classes/Tasks/GameplayTask_SpawnActor.h b/Engine/Source/Runtime/GameplayTasks/Classes/Tasks/GameplayTask_SpawnActor.h
--- a/Engine/Source/Runtime/GameplayTasks/Classes/Tasks/GameplayTask_SpawnActor.h
+++ b/Engine/Source/Runtime/GameplayTasks/Classes/Tasks/GameplayTask_SpawnActor.h
@@ -31,7 +31,7 @@ public:
 	FSpawnActorDelegate	DidNotSpawn;
 	
 	/** Spawn new Actor on the network authority (server) */
-	UFUNCTION(BlueprintCallable, Category = "GameplayTasks", meta = (AdvancedDisplay = "TaskOwner, bSpawnOnlyOnAuthority", DefaultToSelf = "TaskOwner", BlueprintInternalUseOnly = "TRUE"))
+	UFUNCTION(BlueprintCallable, Category = "GameplayTasks", meta = (DisplayName="Spawn Actor for Gameplay Task", AdvancedDisplay = "TaskOwner, bSpawnOnlyOnAuthority", DefaultToSelf = "TaskOwner", BlueprintInternalUseOnly = "TRUE"))
 	static UGameplayTask_SpawnActor* SpawnActor(TScriptInterface<IGameplayTaskOwnerInterface> TaskOwner, FVector SpawnLocation, FRotator SpawnRotation, TSubclassOf<AActor> Class, bool bSpawnOnlyOnAuthority = false);
 
 	UFUNCTION(BlueprintCallable, meta = (HidePin = "WorldContextObject", DefaultToSelf = "WorldContextObject", BlueprintInternalUseOnly = "true"), Category = "GameplayTasks")
