make InitTask virtual

diff --git a/Engine/Source/Runtime/GameplayTasks/Classes/GameplayTask.h b/Engine/Source/Runtime/GameplayTasks/Classes/GameplayTask.h
--- a/Engine/Source/Runtime/GameplayTasks/Classes/GameplayTask.h
+++ b/Engine/Source/Runtime/GameplayTasks/Classes/GameplayTask.h
@@ -52,7 +52,7 @@ protected:
 	virtual void Activate();
 
 	/** Initailizes the task with the task owner interface instance but does not actviate until Activate() is called */
-	void InitTask(IGameplayTaskOwnerInterface& InTaskOwner);
+	virtual void InitTask(IGameplayTaskOwnerInterface& InTaskOwner);
 
 public:
 	virtual void InitSimulatedTask(UGameplayTasksComponent& InGameplayTasksComponent);
