Agged GameplayTasks information to GameplayDebugger #UE4

diff --git a/Engine/Source/Developer/GameplayDebugger/Classes/GameplayDebuggingComponent.h b/Engine/Source/Developer/GameplayDebugger/Classes/GameplayDebuggingComponent.h
--- a/Engine/Source/Developer/GameplayDebugger/Classes/GameplayDebuggingComponent.h
+++ b/Engine/Source/Developer/GameplayDebugger/Classes/GameplayDebuggingComponent.h
@@ -85,6 +85,9 @@ class GAMEPLAYDEBUGGER_API UGameplayDebuggingComponent : public UPrimitiveCompon
 	UPROPERTY(Replicated)
 	FString CurrentAIAssets;
 
+	UPROPERTY(Replicated)
+	FString GameplayTasksState;
+
 	UPROPERTY(Replicated)
 	FString NavDataInfo;
 
diff --git a/Engine/Source/Developer/GameplayDebugger/Private/GameplayDebuggingComponent.cpp b/Engine/Source/Developer/GameplayDebugger/Private/GameplayDebuggingComponent.cpp
--- a/Engine/Source/Developer/GameplayDebugger/Private/GameplayDebuggingComponent.cpp
+++ b/Engine/Source/Developer/GameplayDebugger/Private/GameplayDebuggingComponent.cpp
@@ -214,6 +214,8 @@ void UGameplayDebuggingComponent::GetLifetimeReplicatedProps( TArray< FLifetimeP
 	DOREPLIFETIME(UGameplayDebuggingComponent, CurrentAIState);
 	DOREPLIFETIME(UGameplayDebuggingComponent, CurrentAIAssets);
 
+	DOREPLIFETIME(UGameplayDebuggingComponent, GameplayTasksState);
+
 	DOREPLIFETIME(UGameplayDebuggingComponent, bIsUsingAbilities);
 	DOREPLIFETIME(UGameplayDebuggingComponent, AbilityInfo);
 
@@ -531,6 +533,16 @@ void UGameplayDebuggingComponent::CollectBasicBehaviorData(APawn* MyPawn)
 		CurrentAIState = TEXT("");
 		CurrentAIAssets = TEXT("");
 	}
+
+	UGameplayTasksComponent* GTComponent = MyPawn->FindComponentByClass<UGameplayTasksComponent>();
+	if (GTComponent)
+	{
+		GameplayTasksState = FString::Printf(TEXT("Ticking Tasks: %s\nTask Queue: %s"), *GTComponent->GetTickingTasksDescription(), *GTComponent->GetTasksPriorityQueueDescription());
+	}
+	else
+	{
+		GameplayTasksState = TEXT("");
+	}
 #endif //!(UE_BUILD_SHIPPING || UE_BUILD_TEST)
 }
 
diff --git a/Engine/Source/Developer/GameplayDebugger/Private/GameplayDebuggingHUDComponent.cpp b/Engine/Source/Developer/GameplayDebugger/Private/GameplayDebuggingHUDComponent.cpp
--- a/Engine/Source/Developer/GameplayDebugger/Private/GameplayDebuggingHUDComponent.cpp
+++ b/Engine/Source/Developer/GameplayDebugger/Private/GameplayDebuggingHUDComponent.cpp
@@ -395,7 +395,7 @@ void AGameplayDebuggingHUDComponent::DrawBasicData(APlayerController* PC, class
 		PrintString(DefaultContext, FString::Printf(TEXT("Behavior: {yellow}%s{white}, Tree: {yellow}%s\n"), *DebugComponent->CurrentAIState, *DebugComponent->CurrentAIAssets));
 		PrintString(DefaultContext, FString::Printf(TEXT("Active task: {yellow}%s\n"), *DebugComponent->CurrentAITask));
 	}
-
+		
 	// ability + animation
 	if (DebugComponent->bIsUsingAbilities && DebugComponent->bIsUsingCharacter)
 	{
@@ -410,6 +410,9 @@ void AGameplayDebuggingHUDComponent::DrawBasicData(APlayerController* PC, class
 		PrintString(DefaultContext, FString::Printf(TEXT("Ability: {yellow}%s\n"), *DebugComponent->AbilityInfo));
 	}
 
+	// putting gameplay tasks' stuff last since it can expand heavily
+	PrintString(DefaultContext, FString::Printf(TEXT("GameplayTasks:\n{yellow}%s\n"), *DebugComponent->GameplayTasksState));
+
 	DrawPath(PC, DebugComponent);
 #endif //!(UE_BUILD_SHIPPING || UE_BUILD_TEST)
 }
diff --git a/Engine/Source/Runtime/GameplayTasks/Classes/GameplayTasksComponent.h b/Engine/Source/Runtime/GameplayTasks/Classes/GameplayTasksComponent.h
--- a/Engine/Source/Runtime/GameplayTasks/Classes/GameplayTasksComponent.h
+++ b/Engine/Source/Runtime/GameplayTasks/Classes/GameplayTasksComponent.h
@@ -116,7 +116,9 @@ public:
 	static EGameplayTaskRunResult K2_RunGameplayTask(TScriptInterface<IGameplayTaskOwnerInterface> TaskOwner, UGameplayTask* Task, uint8 Priority, TArray<TSubclassOf<UGameplayTaskResource> > AdditionalRequiredResources, TArray<TSubclassOf<UGameplayTaskResource> > AdditionalClaimedResources);
 
 	static EGameplayTaskRunResult RunGameplayTask(IGameplayTaskOwnerInterface& TaskOwner, UGameplayTask& Task, uint8 Priority, FGameplayResourceSet AdditionalRequiredResources, FGameplayResourceSet AdditionalClaimedResources);
-
+	
+	FString GetTickingTasksDescription() const;
+	FString GetTasksPriorityQueueDescription() const;
 #if ENABLE_VISUAL_LOG
 	static FString GetTaskStateName(EGameplayTaskState Value);
 	void DescribeSelfToVisLog(struct FVisualLogEntry* Snapshot) const;
diff --git a/Engine/Source/Runtime/GameplayTasks/Private/GameplayTasksComponent.cpp b/Engine/Source/Runtime/GameplayTasks/Private/GameplayTasksComponent.cpp
--- a/Engine/Source/Runtime/GameplayTasks/Private/GameplayTasksComponent.cpp
+++ b/Engine/Source/Runtime/GameplayTasks/Private/GameplayTasksComponent.cpp
@@ -428,47 +428,56 @@ void UGameplayTasksComponent::SetCurrentlyClaimedResources(FGameplayResourceSet
 //----------------------------------------------------------------------//
 // debugging
 //----------------------------------------------------------------------//
-#if ENABLE_VISUAL_LOG
-void UGameplayTasksComponent::DescribeSelfToVisLog(FVisualLogEntry* Snapshot) const
+FString UGameplayTasksComponent::GetTickingTasksDescription() const
 {
-	static const FString CategoryName = TEXT("GameplayTasks");
-	static const FString TickingTasksName = TEXT("Ticking tasks");
-	static const FString PriorityQueueName = TEXT("Priority Queue");
-
-	if (IsPendingKill())
-	{
-		return;
-	}
-
-	FVisualLogStatusCategory StatusCategory(CategoryName);
-
 	FString TasksDescription;
 	for (auto& Task : TickingTasks)
 	{
 		if (Task.IsValid())
 		{
-			TasksDescription += FString::Printf(TEXT("%s %s\n"), *GetTaskStateName(Task->GetState()), *Task->GetDebugDescription());
+			TasksDescription += FString::Printf(TEXT("\n%s %s"), *GetTaskStateName(Task->GetState()), *Task->GetDebugDescription());
 		}
 		else
 		{
-			TasksDescription += TEXT("NULL\n");
+			TasksDescription += TEXT("\nNULL");
 		}
 	}
-	StatusCategory.Add(TickingTasksName, TasksDescription);
+	return TasksDescription;
+}
 
-	TasksDescription.Reset();
+FString UGameplayTasksComponent::GetTasksPriorityQueueDescription() const
+{
+	FString TasksDescription;
 	for (auto Task : TaskPriorityQueue)
 	{
 		if (Task != nullptr)
 		{
-			TasksDescription += FString::Printf(TEXT("%s %s\n"), *GetTaskStateName(Task->GetState()), *Task->GetDebugDescription());
+			TasksDescription += FString::Printf(TEXT("\n%s %s"), *GetTaskStateName(Task->GetState()), *Task->GetDebugDescription());
 		}
 		else
 		{
-			TasksDescription += TEXT("NULL\n");
+			TasksDescription += TEXT("\nNULL");
 		}
 	}
-	StatusCategory.Add(PriorityQueueName, TasksDescription);
+	return TasksDescription;
+}
+
+#if ENABLE_VISUAL_LOG
+void UGameplayTasksComponent::DescribeSelfToVisLog(FVisualLogEntry* Snapshot) const
+{
+	static const FString CategoryName = TEXT("GameplayTasks");
+	static const FString TickingTasksName = TEXT("Ticking tasks");
+	static const FString PriorityQueueName = TEXT("Priority Queue");
+
+	if (IsPendingKill())
+	{
+		return;
+	}
+
+	FVisualLogStatusCategory StatusCategory(CategoryName);
+
+	StatusCategory.Add(TickingTasksName, GetTickingTasksDescription());
+	StatusCategory.Add(PriorityQueueName, GetTasksPriorityQueueDescription());
 
 	Snapshot->Status.Add(StatusCategory);
 }
