Fixed no-VLogging builds compilation issues #UE4

diff --git a/Engine/Source/Runtime/GameplayTasks/Classes/GameplayTask.h b/Engine/Source/Runtime/GameplayTasks/Classes/GameplayTask.h
--- a/Engine/Source/Runtime/GameplayTasks/Classes/GameplayTask.h
+++ b/Engine/Source/Runtime/GameplayTasks/Classes/GameplayTask.h
@@ -324,7 +324,7 @@ protected:
 
 	TWeakObjectPtr<UGameplayTasksComponent>	TasksComponent;
 
-#if ENABLE_VISUAL_LOG
+#if !(UE_BUILD_SHIPPING || UE_BUILD_TEST)
 	mutable FString DebugDescription;
 public:
 	const FString& GetDebugDescription() const
@@ -337,7 +337,7 @@ public:
 	}
 	virtual FString GenerateDebugDescription() const;
 	FString GetTaskStateName() const;
-#endif // ENABLE_VISUAL_LOG
+#endif // !(UE_BUILD_SHIPPING || UE_BUILD_TEST)
 };
 
 template <class T>
diff --git a/Engine/Source/Runtime/GameplayTasks/Classes/GameplayTasksComponent.h b/Engine/Source/Runtime/GameplayTasks/Classes/GameplayTasksComponent.h
--- a/Engine/Source/Runtime/GameplayTasks/Classes/GameplayTasksComponent.h
+++ b/Engine/Source/Runtime/GameplayTasks/Classes/GameplayTasksComponent.h
@@ -117,8 +117,10 @@ public:
 
 	static EGameplayTaskRunResult RunGameplayTask(IGameplayTaskOwnerInterface& TaskOwner, UGameplayTask& Task, uint8 Priority, FGameplayResourceSet AdditionalRequiredResources, FGameplayResourceSet AdditionalClaimedResources);
 	
+#if !(UE_BUILD_SHIPPING || UE_BUILD_TEST)
 	FString GetTickingTasksDescription() const;
 	FString GetTasksPriorityQueueDescription() const;
+#endif // !(UE_BUILD_SHIPPING || UE_BUILD_TEST)
 #if ENABLE_VISUAL_LOG
 	static FString GetTaskStateName(EGameplayTaskState Value);
 	void DescribeSelfToVisLog(struct FVisualLogEntry* Snapshot) const;
diff --git a/Engine/Source/Runtime/GameplayTasks/Private/GameplayTask.cpp b/Engine/Source/Runtime/GameplayTasks/Private/GameplayTask.cpp
--- a/Engine/Source/Runtime/GameplayTasks/Private/GameplayTask.cpp
+++ b/Engine/Source/Runtime/GameplayTasks/Private/GameplayTask.cpp
@@ -360,7 +360,7 @@ void UGameplayTask::PauseInTaskQueue()
 	}
 }
 
-#if ENABLE_VISUAL_LOG
+#if !(UE_BUILD_SHIPPING || UE_BUILD_TEST)
 //----------------------------------------------------------------------//
 // debug
 //----------------------------------------------------------------------//
@@ -378,4 +378,4 @@ FString UGameplayTask::GetTaskStateName() const
 	return Enum->GetEnumName(int32(TaskState));
 }
 
-#endif // ENABLE_VISUAL_LOG
\ No newline at end of file
+#endif // !(UE_BUILD_SHIPPING || UE_BUILD_TEST)
\ No newline at end of file
diff --git a/Engine/Source/Runtime/GameplayTasks/Private/GameplayTasksComponent.cpp b/Engine/Source/Runtime/GameplayTasks/Private/GameplayTasksComponent.cpp
--- a/Engine/Source/Runtime/GameplayTasks/Private/GameplayTasksComponent.cpp
+++ b/Engine/Source/Runtime/GameplayTasks/Private/GameplayTasksComponent.cpp
@@ -428,6 +428,7 @@ void UGameplayTasksComponent::SetCurrentlyClaimedResources(FGameplayResourceSet
 //----------------------------------------------------------------------//
 // debugging
 //----------------------------------------------------------------------//
+#if !(UE_BUILD_SHIPPING || UE_BUILD_TEST)
 FString UGameplayTasksComponent::GetTickingTasksDescription() const
 {
 	FString TasksDescription;
@@ -461,6 +462,7 @@ FString UGameplayTasksComponent::GetTasksPriorityQueueDescription() const
 	}
 	return TasksDescription;
 }
+#endif // !(UE_BUILD_SHIPPING || UE_BUILD_TEST)
 
 #if ENABLE_VISUAL_LOG
 void UGameplayTasksComponent::DescribeSelfToVisLog(FVisualLogEntry* Snapshot) const
