Fixed running resource-less GameplayTask with additionally specified claimed resources not resulting in tasks' resource management kicking in #UE4

diff --git a/Engine/Source/Runtime/GameplayTasks/Classes/GameplayTask.h b/Engine/Source/Runtime/GameplayTasks/Classes/GameplayTask.h
--- a/Engine/Source/Runtime/GameplayTasks/Classes/GameplayTask.h
+++ b/Engine/Source/Runtime/GameplayTasks/Classes/GameplayTask.h
@@ -214,7 +214,7 @@ public:
 	FORCEINLINE bool IsSimulating() const { return (bIsSimulating != 0); }
 	FORCEINLINE bool IsPausable() const { return (bIsPausable != 0); }
 	FORCEINLINE uint8 GetPriority() const { return Priority; }
-	FORCEINLINE bool RequiresPriorityOrResourceManagement() const { return bCaresAboutPriority == true || RequiredResources.IsEmpty() == false; }
+	FORCEINLINE bool RequiresPriorityOrResourceManagement() const { return bCaresAboutPriority == true || RequiredResources.IsEmpty() == false || ClaimedResources.IsEmpty() == false; }
 	FORCEINLINE FGameplayResourceSet GetRequiredResources() { return RequiredResources; }
 	FORCEINLINE FGameplayResourceSet GetClaimedResources() { return ClaimedResources; }
 	
diff --git a/Engine/Source/Runtime/GameplayTasks/Classes/GameplayTasksComponent.h b/Engine/Source/Runtime/GameplayTasks/Classes/GameplayTasksComponent.h
--- a/Engine/Source/Runtime/GameplayTasks/Classes/GameplayTasksComponent.h
+++ b/Engine/Source/Runtime/GameplayTasks/Classes/GameplayTasksComponent.h
@@ -112,7 +112,7 @@ public:
 
 	FORCEINLINE FGameplayResourceSet GetCurrentlyUsedResources() const { return CurrentlyClaimedResources; }
 
-	UFUNCTION(BlueprintCallable, DisplayName="Run Gameplay Task", Category = "Gameplay Tasks", meta = (AutoCreateRefTerm = "AdditionalResourcesToBeLockedWhenRun, AdditionalClaimedResources", AdvancedDisplay = "AdditionalResourcesToBeLockedWhenRun, AdditionalClaimedResources"))
+	UFUNCTION(BlueprintCallable, DisplayName="Run Gameplay Task", Category = "Gameplay Tasks", meta = (AutoCreateRefTerm = "AdditionalRequiredResources, AdditionalClaimedResources", AdvancedDisplay = "AdditionalRequiredResources, AdditionalClaimedResources"))
 	static EGameplayTaskRunResult K2_RunGameplayTask(TScriptInterface<IGameplayTaskOwnerInterface> TaskOwner, UGameplayTask* Task, uint8 Priority, TArray<TSubclassOf<UGameplayTaskResource> > AdditionalRequiredResources, TArray<TSubclassOf<UGameplayTaskResource> > AdditionalClaimedResources);
 
 	static EGameplayTaskRunResult RunGameplayTask(IGameplayTaskOwnerInterface& TaskOwner, UGameplayTask& Task, uint8 Priority, FGameplayResourceSet AdditionalRequiredResources, FGameplayResourceSet AdditionalClaimedResources);
diff --git a/Engine/Source/Runtime/GameplayTasks/Private/GameplayTask.cpp b/Engine/Source/Runtime/GameplayTasks/Private/GameplayTask.cpp
--- a/Engine/Source/Runtime/GameplayTasks/Private/GameplayTask.cpp
+++ b/Engine/Source/Runtime/GameplayTasks/Private/GameplayTask.cpp
@@ -214,8 +214,11 @@ void UGameplayTask::AddRequiredResourceSet(const TArray<TSubclassOf<UGameplayTas
 {
 	for (auto Resource : RequiredResourceSet)
 	{
-		const uint8 ResourceID = UGameplayTaskResource::GetResourceID(Resource);
-		RequiredResources.AddID(ResourceID);
+		if (Resource)
+		{
+			const uint8 ResourceID = UGameplayTaskResource::GetResourceID(Resource);
+			RequiredResources.AddID(ResourceID);
+		}
 	}
 }
 
@@ -235,7 +238,10 @@ void UGameplayTask::AddClaimedResourceSet(const TArray<TSubclassOf<UGameplayTask
 {
 	for (auto ResourceClass : AdditionalResourcesToClaim)
 	{
-		ClaimedResources.AddID(UGameplayTaskResource::GetResourceID(ResourceClass));
+		if (ResourceClass)
+		{
+			ClaimedResources.AddID(UGameplayTaskResource::GetResourceID(ResourceClass));
+		}
 	}
 }
 
