Added a safety check to make sure GameplayTasks won't get double-activated #UE4

diff --git a/Engine/Source/Runtime/GameplayTasks/Private/GameplayTask.cpp b/Engine/Source/Runtime/GameplayTasks/Private/GameplayTask.cpp
--- a/Engine/Source/Runtime/GameplayTasks/Private/GameplayTask.cpp
+++ b/Engine/Source/Runtime/GameplayTasks/Private/GameplayTask.cpp
@@ -252,6 +252,14 @@ void UGameplayTask::AddClaimedResourceSet(FGameplayResourceSet AdditionalResourc
 
 void UGameplayTask::PerformActivation()
 {
+	if (TaskState == EGameplayTaskState::Active)
+	{
+		UE_VLOG(GetGameplayTasksComponent(), LogGameplayTasks, Warning
+			, TEXT("%s PerformActivation called while TaskState is already Active. Bailing out.")
+			, *GetName());
+		return;
+	}
+
 	TaskState = EGameplayTaskState::Active;
 
 	Activate();
