CircularDependency fix.

diff --git a/Engine/Source/Runtime/GameplayTasks/GameplayTasks.Build.cs b/Engine/Source/Runtime/GameplayTasks/GameplayTasks.Build.cs
--- a/Engine/Source/Runtime/GameplayTasks/GameplayTasks.Build.cs
+++ b/Engine/Source/Runtime/GameplayTasks/GameplayTasks.Build.cs
@@ -48,6 +48,7 @@ namespace UnrealBuildTool.Rules
 			if (UEBuildConfiguration.bBuildEditor == true)
 			{
 				PrivateDependencyModuleNames.Add("UnrealEd");
+				CircularlyReferencedDependentModules.Add("UnrealEd");
                 //PrivateDependencyModuleNames.Add("GameplayTagsEditor");
 			}
 		}
