Fixed a bug in resource-claiming GameplayTasks handling that was resulting in not releasing resources if task was being ended as part of Activate function call #UE4

Merged over from Orion

diff --git a/Engine/Source/Runtime/GameplayTasks/Private/GameplayTasksComponent.cpp b/Engine/Source/Runtime/GameplayTasks/Private/GameplayTasksComponent.cpp
--- a/Engine/Source/Runtime/GameplayTasks/Private/GameplayTasksComponent.cpp
+++ b/Engine/Source/Runtime/GameplayTasks/Private/GameplayTasksComponent.cpp
@@ -63,8 +63,7 @@ void UGameplayTasksComponent::OnTaskDeactivated(UGameplayTask& Task)
 void UGameplayTasksComponent::OnTaskEnded(UGameplayTask& Task)
 {
 	ensure(Task.RequiresPriorityOrResourceManagement() == true);
-
-	RemoveTaskFromPriorityQueue(Task);
+	RemoveResourceConsumingTask(Task);
 }
 
 void UGameplayTasksComponent::GetLifetimeReplicatedProps(TArray< FLifetimeProperty >& OutLifetimeProps) const
@@ -188,7 +187,7 @@ AActor* UGameplayTasksComponent::GetAvatarActor(const UGameplayTask* Task) const
 //----------------------------------------------------------------------//
 void UGameplayTasksComponent::AddTaskReadyForActivation(UGameplayTask& NewTask)
 {
-	UE_VLOG(this, LogGameplayTasks, Log, TEXT("RemoveResourceConsumingTask %s"), *NewTask.GetName());
+	UE_VLOG(this, LogGameplayTasks, Log, TEXT("AddTaskReadyForActivation %s"), *NewTask.GetName());
 
 	ensure(NewTask.RequiresPriorityOrResourceManagement() == true);
 	
