Copying //UE4/Dev-Framework to Dev-Main (//UE4/Dev-Main) @ 2926658

diff --git a/Engine/Source/Runtime/Foliage/Private/FoliageComponent.cpp b/Engine/Source/Runtime/Foliage/Private/FoliageComponent.cpp
--- a/Engine/Source/Runtime/Foliage/Private/FoliageComponent.cpp
+++ b/Engine/Source/Runtime/Foliage/Private/FoliageComponent.cpp
@@ -97,7 +97,7 @@ float AInteractiveFoliageActor::TakeDamage(float DamageAmount, FDamageEvent cons
 	return 0.f;
 }
 
-void AInteractiveFoliageActor::CapsuleTouched(AActor* Other, UPrimitiveComponent* OtherComp, int32 OtherBodyIndex, bool bFromSweep, const FHitResult& SweepResult)
+void AInteractiveFoliageActor::CapsuleTouched(UPrimitiveComponent* OverlappedComp, AActor* Other, UPrimitiveComponent* OtherComp, int32 OtherBodyIndex, bool bFromSweep, const FHitResult& SweepResult)
 {
 	if (Other != NULL && OtherComp != NULL)
 	{
@@ -137,12 +137,11 @@ void AInteractiveFoliageActor::Tick(float DeltaSeconds)
 	// The proxy will be NULL on dedicated server
 	if (FoliageComponent->IsRegistered() && FoliageComponent->FoliageSceneProxy)
 	{
-		TArray<AActor*> Touching;
+		TSet<AActor*> Touching;
 		GetOverlappingActors(Touching);
 
-		for (int32 TouchingIndex = 0; TouchingIndex < Touching.Num(); TouchingIndex++)
+		for (AActor* TouchingActor : Touching)
 		{
-			AActor* TouchingActor = Touching[TouchingIndex];
 			if (TouchingActor != NULL && TouchingActor->GetRootComponent() != NULL)
 			{
 				const FVector TouchingActorPosition(TouchingActor->GetRootComponent()->Bounds.Origin.X, TouchingActor->GetRootComponent()->Bounds.Origin.Y, GetRootComponent()->Bounds.Origin.Z);
diff --git a/Engine/Source/Runtime/Foliage/Public/InteractiveFoliageActor.h b/Engine/Source/Runtime/Foliage/Public/InteractiveFoliageActor.h
--- a/Engine/Source/Runtime/Foliage/Public/InteractiveFoliageActor.h
+++ b/Engine/Source/Runtime/Foliage/Public/InteractiveFoliageActor.h
@@ -85,7 +85,7 @@ protected:
 
 	/** Called when capsule is touched */
 	UFUNCTION()
-	void CapsuleTouched(AActor* Other, UPrimitiveComponent* OtherComp, int32 OtherBodyIndex, bool bFromSweep, const FHitResult& OverlapInfo);
+	void CapsuleTouched(UPrimitiveComponent* OverlappedComp, AActor* Other, UPrimitiveComponent* OtherComp, int32 OtherBodyIndex, bool bFromSweep, const FHitResult& OverlapInfo);
 public:
 	//~ Begin AActor Interface
 	virtual void Tick(float DeltaSeconds) override;
diff --git a/Engine/Source/Runtime/GameplayAbilities/Private/Abilities/GameplayAbility.cpp b/Engine/Source/Runtime/GameplayAbilities/Private/Abilities/GameplayAbility.cpp
--- a/Engine/Source/Runtime/GameplayAbilities/Private/Abilities/GameplayAbility.cpp
+++ b/Engine/Source/Runtime/GameplayAbilities/Private/Abilities/GameplayAbility.cpp
@@ -1068,19 +1068,19 @@ UGameplayTasksComponent* UGameplayAbility::GetGameplayTasksComponent(const UGame
 	return GetCurrentActorInfo() ? GetCurrentActorInfo()->AbilitySystemComponent.Get() : nullptr;
 }
 
-AActor* UGameplayAbility::GetOwnerActor(const UGameplayTask* Task) const
+AActor* UGameplayAbility::GetGameplayTaskOwner(const UGameplayTask* Task) const
 {
 	const FGameplayAbilityActorInfo* Info = GetCurrentActorInfo();
 	return Info ? Info->OwnerActor.Get() : nullptr;
 }
 
-AActor* UGameplayAbility::GetAvatarActor(const UGameplayTask* Task) const
+AActor* UGameplayAbility::GetGameplayTaskAvatar(const UGameplayTask* Task) const
 {
 	const FGameplayAbilityActorInfo* Info = GetCurrentActorInfo();
 	return Info ? Info->AvatarActor.Get() : nullptr;
 }
 
-void UGameplayAbility::OnTaskInitialized(UGameplayTask& Task)
+void UGameplayAbility::OnGameplayTaskInitialized(UGameplayTask& Task)
 {
 	UAbilityTask* AbilityTask = Cast<UAbilityTask>(&Task);
 	if (AbilityTask)
@@ -1090,13 +1090,25 @@ void UGameplayAbility::OnTaskInitialized(UGameplayTask& Task)
 	}
 }
 
-void UGameplayAbility::OnTaskActivated(UGameplayTask& Task)
+void UGameplayAbility::OnGameplayTaskActivated(UGameplayTask& Task)
 {
 	ABILITY_VLOG(CastChecked<AActor>(GetOuter()), Log, TEXT("Task Started %s"), *Task.GetName());
 
 	ActiveTasks.Add(&Task);
 }
 
+void UGameplayAbility::OnGameplayTaskDeactivated(UGameplayTask& Task)
+{
+	ABILITY_VLOG(CastChecked<AActor>(GetOuter()), Log, TEXT("Task Ended %s"), *Task.GetName());
+
+	ActiveTasks.Remove(&Task);
+
+	if (ENABLE_ABILITYTASK_DEBUGMSG)
+	{
+		AddAbilityTaskDebugMessage(&Task, TEXT("Ended."));
+	}
+}
+
 void UGameplayAbility::ConfirmTaskByInstanceName(FName InstanceName, bool bEndTask)
 {
 	TArray<UGameplayTask*, TInlineAllocator<8> > NamedTasks;
@@ -1204,18 +1216,6 @@ void UGameplayAbility::EndAbilityState(FName OptionalStateNameToEnd)
 	}
 }
 
-void UGameplayAbility::OnTaskDeactivated(UGameplayTask& Task)
-{
-	ABILITY_VLOG(CastChecked<AActor>(GetOuter()), Log, TEXT("Task Ended %s"), *Task.GetName());
-
-	ActiveTasks.Remove(&Task);
-
-	if (ENABLE_ABILITYTASK_DEBUGMSG)
-	{
-		AddAbilityTaskDebugMessage(&Task, TEXT("Ended."));
-	}
-}
-
 void UGameplayAbility::AddAbilityTaskDebugMessage(UGameplayTask* AbilityTask, FString DebugMessage)
 {
 	TaskDebugMessages.AddDefaulted();
diff --git a/Engine/Source/Runtime/GameplayAbilities/Private/Abilities/Tasks/AbilityTask_WaitOverlap.cpp b/Engine/Source/Runtime/GameplayAbilities/Private/Abilities/Tasks/AbilityTask_WaitOverlap.cpp
--- a/Engine/Source/Runtime/GameplayAbilities/Private/Abilities/Tasks/AbilityTask_WaitOverlap.cpp
+++ b/Engine/Source/Runtime/GameplayAbilities/Private/Abilities/Tasks/AbilityTask_WaitOverlap.cpp
@@ -9,16 +9,7 @@ UAbilityTask_WaitOverlap::UAbilityTask_WaitOverlap(const FObjectInitializer& Obj
 
 }
 
-
-void UAbilityTask_WaitOverlap::OnOverlapCallback(AActor* OtherActor, class UPrimitiveComponent* OtherComp, int32 OtherBodyIndex, bool bFromSweep, const FHitResult& SweepResult)
-{
-	if (OtherActor && bFromSweep)
-	{	
-		//OnOverlap.Broadcast(SweepResult);
-	}
-}
-
-void UAbilityTask_WaitOverlap::OnHitCallback(AActor* OtherActor, UPrimitiveComponent* OtherComp, FVector NormalImpulse, const FHitResult& Hit)
+void UAbilityTask_WaitOverlap::OnHitCallback(UPrimitiveComponent* HitComp, AActor* OtherActor, UPrimitiveComponent* OtherComp, FVector NormalImpulse, const FHitResult& Hit)
 {
 	if(OtherActor)
 	{
@@ -54,7 +45,6 @@ void UAbilityTask_WaitOverlap::Activate()
 	UPrimitiveComponent* PrimComponent = GetComponent();
 	if (PrimComponent)
 	{
-		PrimComponent->OnComponentBeginOverlap.AddDynamic(this, &UAbilityTask_WaitOverlap::OnOverlapCallback);
 		PrimComponent->OnComponentHit.AddDynamic(this, &UAbilityTask_WaitOverlap::OnHitCallback);
 	}
 }
@@ -64,7 +54,6 @@ void UAbilityTask_WaitOverlap::OnDestroy(bool AbilityEnded)
 	UPrimitiveComponent* PrimComponent = GetComponent();
 	if (PrimComponent)
 	{
-		PrimComponent->OnComponentBeginOverlap.RemoveDynamic(this, &UAbilityTask_WaitOverlap::OnOverlapCallback);
 		PrimComponent->OnComponentHit.RemoveDynamic(this, &UAbilityTask_WaitOverlap::OnHitCallback);
 	}
 
diff --git a/Engine/Source/Runtime/GameplayAbilities/Private/AbilitySystemComponent.cpp b/Engine/Source/Runtime/GameplayAbilities/Private/AbilitySystemComponent.cpp
--- a/Engine/Source/Runtime/GameplayAbilities/Private/AbilitySystemComponent.cpp
+++ b/Engine/Source/Runtime/GameplayAbilities/Private/AbilitySystemComponent.cpp
@@ -1380,7 +1380,13 @@ TArray<UGameplayTask*>&	UAbilitySystemComponent::GetAbilityActiveTasks(UGameplay
 	return Ability->ActiveTasks;
 }
 
-AActor* UAbilitySystemComponent::GetAvatarActor(const UGameplayTask* Task) const
+AActor* UAbilitySystemComponent::GetGameplayTaskAvatar(const UGameplayTask* Task) const
+{
+	check(AbilityActorInfo.IsValid());
+	return AbilityActorInfo->AvatarActor.Get();
+}
+
+AActor* UAbilitySystemComponent::GetAvatarActor() const
 {
 	check(AbilityActorInfo.IsValid());
 	return AbilityActorInfo->AvatarActor.Get();
diff --git a/Engine/Source/Runtime/GameplayAbilities/Private/GameplayCueManager.cpp b/Engine/Source/Runtime/GameplayAbilities/Private/GameplayCueManager.cpp
--- a/Engine/Source/Runtime/GameplayAbilities/Private/GameplayCueManager.cpp
+++ b/Engine/Source/Runtime/GameplayAbilities/Private/GameplayCueManager.cpp
@@ -159,7 +159,7 @@ void UGameplayCueManager::EndGameplayCuesFor(AActor* TargetActor)
 			AGameplayCueNotify_Actor* InstancedCue = It.Value().Get();
 			if (InstancedCue)
 			{
-				InstancedCue->OnOwnerDestroyed();
+				InstancedCue->OnOwnerDestroyed(TargetActor);
 			}
 			It.RemoveCurrent();
 		}
diff --git a/Engine/Source/Runtime/GameplayAbilities/Private/GameplayCueNotify_Actor.cpp b/Engine/Source/Runtime/GameplayAbilities/Private/GameplayCueNotify_Actor.cpp
--- a/Engine/Source/Runtime/GameplayAbilities/Private/GameplayCueNotify_Actor.cpp
+++ b/Engine/Source/Runtime/GameplayAbilities/Private/GameplayCueNotify_Actor.cpp
@@ -194,7 +194,7 @@ void AGameplayCueNotify_Actor::HandleGameplayCue(AActor* MyTarget, EGameplayCueE
 	}
 }
 
-void AGameplayCueNotify_Actor::OnOwnerDestroyed()
+void AGameplayCueNotify_Actor::OnOwnerDestroyed(AActor* DestroyedActor)
 {
 	if (bInRecycleQueue)
 	{
diff --git a/Engine/Source/Runtime/GameplayAbilities/Public/Abilities/GameplayAbility.h b/Engine/Source/Runtime/GameplayAbilities/Public/Abilities/GameplayAbility.h
--- a/Engine/Source/Runtime/GameplayAbilities/Public/Abilities/GameplayAbility.h
+++ b/Engine/Source/Runtime/GameplayAbilities/Public/Abilities/GameplayAbility.h
@@ -267,17 +267,12 @@ public:
 	// --------------------------------------
 	//	IGameplayTaskOwnerInterface
 	// --------------------------------------	
-	/** setup additional properties if given task is an AbilityTask */
-	virtual void OnTaskInitialized(UGameplayTask& Task) override;
-	/** Called by an ability task, originating from this ability, when it starts */
-	virtual void OnTaskActivated(UGameplayTask& Task) override;
-
-	/** Called by an ability task, originating from this ability, when it ends */
-	virtual void OnTaskDeactivated(UGameplayTask& Task) override;
-
 	virtual UGameplayTasksComponent* GetGameplayTasksComponent(const UGameplayTask& Task) const override;
-	virtual AActor* GetOwnerActor(const UGameplayTask* Task) const override;
-	virtual AActor* GetAvatarActor(const UGameplayTask* Task) const override;
+	virtual AActor* GetGameplayTaskOwner(const UGameplayTask* Task) const override;
+	virtual AActor* GetGameplayTaskAvatar(const UGameplayTask* Task) const override;
+	virtual void OnGameplayTaskInitialized(UGameplayTask& Task) override;
+	virtual void OnGameplayTaskActivated(UGameplayTask& Task) override;
+	virtual void OnGameplayTaskDeactivated(UGameplayTask& Task) override;
 
 	// --------------------------------------
 	//	Input
diff --git a/Engine/Source/Runtime/GameplayAbilities/Public/Abilities/Tasks/AbilityTask.h b/Engine/Source/Runtime/GameplayAbilities/Public/Abilities/Tasks/AbilityTask.h
--- a/Engine/Source/Runtime/GameplayAbilities/Public/Abilities/Tasks/AbilityTask.h
+++ b/Engine/Source/Runtime/GameplayAbilities/Public/Abilities/Tasks/AbilityTask.h
@@ -115,7 +115,7 @@ class GAMEPLAYABILITIES_API UAbilityTask : public UGameplayTask
 
 		T* MyObj = NewObject<T>();
 		UGameplayAbility* ThisAbility = CastChecked<UGameplayAbility>(WorldContextObject);
-		MyObj->InitTask(*ThisAbility, ThisAbility->GetDefaultPriority());
+		MyObj->InitTask(*ThisAbility, ThisAbility->GetGameplayTaskDefaultPriority());
 		MyObj->InstanceName = InstanceName;
 		return MyObj;
 	}
diff --git a/Engine/Source/Runtime/GameplayAbilities/Public/Abilities/Tasks/AbilityTask_WaitOverlap.h b/Engine/Source/Runtime/GameplayAbilities/Public/Abilities/Tasks/AbilityTask_WaitOverlap.h
--- a/Engine/Source/Runtime/GameplayAbilities/Public/Abilities/Tasks/AbilityTask_WaitOverlap.h
+++ b/Engine/Source/Runtime/GameplayAbilities/Public/Abilities/Tasks/AbilityTask_WaitOverlap.h
@@ -23,10 +23,7 @@ class UAbilityTask_WaitOverlap : public UAbilityTask
 	FWaitOverlapDelegate	OnOverlap;
 
 	UFUNCTION()
-	void OnOverlapCallback(AActor* OtherActor, UPrimitiveComponent* OtherComp, int32 OtherBodyIndex, bool bFromSweep, const FHitResult& SweepResult);
-
-	UFUNCTION()
-	void OnHitCallback(AActor* OtherActor, UPrimitiveComponent* OtherComp, FVector NormalImpulse, const FHitResult& Hit);
+	void OnHitCallback(UPrimitiveComponent* HitComp, AActor* OtherActor, UPrimitiveComponent* OtherComp, FVector NormalImpulse, const FHitResult& Hit);
 
 	virtual void Activate() override;
 
diff --git a/Engine/Source/Runtime/GameplayAbilities/Public/AbilitySystemComponent.h b/Engine/Source/Runtime/GameplayAbilities/Public/AbilitySystemComponent.h
--- a/Engine/Source/Runtime/GameplayAbilities/Public/AbilitySystemComponent.h
+++ b/Engine/Source/Runtime/GameplayAbilities/Public/AbilitySystemComponent.h
@@ -1365,10 +1365,8 @@ public:
 
 	const FActiveGameplayEffect* GetActiveGameplayEffect(const FActiveGameplayEffectHandle Handle) const;
 
-	// -
-	// UGameplayTasksComponent
-	// -
-	virtual AActor* GetAvatarActor(const UGameplayTask* Task) const override;
+	virtual AActor* GetGameplayTaskAvatar(const UGameplayTask* Task) const override;
+	AActor* GetAvatarActor() const;
 
 	/** Suppress all ability granting through GEs on this component */
 	bool bSuppressGrantAbility;
diff --git a/Engine/Source/Runtime/GameplayAbilities/Public/GameplayCueNotify_Actor.h b/Engine/Source/Runtime/GameplayAbilities/Public/GameplayCueNotify_Actor.h
--- a/Engine/Source/Runtime/GameplayAbilities/Public/GameplayCueNotify_Actor.h
+++ b/Engine/Source/Runtime/GameplayAbilities/Public/GameplayCueNotify_Actor.h
@@ -29,7 +29,7 @@ class GAMEPLAYABILITIES_API AGameplayCueNotify_Actor : public AActor
 	virtual bool HandlesEvent(EGameplayCueEvent::Type EventType) const;
 
 	UFUNCTION()
-	virtual void OnOwnerDestroyed();
+	virtual void OnOwnerDestroyed(AActor* DestroyedActor);
 
 	virtual void BeginPlay() override;
 
diff --git a/Engine/Source/Runtime/GameplayTags/Private/GameplayTagsManager.cpp b/Engine/Source/Runtime/GameplayTags/Private/GameplayTagsManager.cpp
--- a/Engine/Source/Runtime/GameplayTags/Private/GameplayTagsManager.cpp
+++ b/Engine/Source/Runtime/GameplayTags/Private/GameplayTagsManager.cpp
@@ -233,9 +233,9 @@ void UGameplayTagsManager::ConstructGameplayTagTree()
 				FName OldTagName = NAME_None;
 				FName NewTagName;
 
-				if (FParse::Value(*It.Value(), TEXT("OldTagName="), OldTagName))
+				if (FParse::Value(*It.Value().GetValue(), TEXT("OldTagName="), OldTagName))
 				{
-					if (FParse::Value(*It.Value(), TEXT("NewTagName="), NewTagName))
+					if (FParse::Value(*It.Value().GetValue(), TEXT("NewTagName="), NewTagName))
 					{
 						if (ensureMsgf(!TagRedirects.Contains(OldTagName), TEXT("Old tag %s is being redirected to more than one tag. Please remove all the redirections except for one."), *OldTagName.ToString()))
 						{
diff --git a/Engine/Source/Runtime/GameplayTasks/Classes/GameplayTask.h b/Engine/Source/Runtime/GameplayTasks/Classes/GameplayTask.h
--- a/Engine/Source/Runtime/GameplayTasks/Classes/GameplayTask.h
+++ b/Engine/Source/Runtime/GameplayTasks/Classes/GameplayTask.h
@@ -131,7 +131,7 @@ public:
 };
 
 UCLASS(Abstract, meta = (ExposedAsyncProxy))
-class GAMEPLAYTASKS_API UGameplayTask : public UObject
+class GAMEPLAYTASKS_API UGameplayTask : public UObject, public IGameplayTaskOwnerInterface
 {
 	GENERATED_BODY()
 
@@ -213,15 +213,18 @@ public:
 	FORCEINLINE bool IsSimulatedTask() const { return (bSimulatedTask != 0); }
 	FORCEINLINE bool IsSimulating() const { return (bIsSimulating != 0); }
 	FORCEINLINE bool IsPausable() const { return (bIsPausable != 0); }
+	FORCEINLINE bool HasOwnerFinished() const {	return (bOwnerFinished != 0); }
 	FORCEINLINE uint8 GetPriority() const { return Priority; }
 	FORCEINLINE bool RequiresPriorityOrResourceManagement() const { return bCaresAboutPriority == true || RequiredResources.IsEmpty() == false || ClaimedResources.IsEmpty() == false; }
-	FORCEINLINE FGameplayResourceSet GetRequiredResources() { return RequiredResources; }
-	FORCEINLINE FGameplayResourceSet GetClaimedResources() { return ClaimedResources; }
+	FORCEINLINE FGameplayResourceSet GetRequiredResources() const { return RequiredResources; }
+	FORCEINLINE FGameplayResourceSet GetClaimedResources() const { return ClaimedResources; }
 	
 	FORCEINLINE EGameplayTaskState GetState() const { return TaskState; }
 	FORCEINLINE bool IsActive() const { return (TaskState == EGameplayTaskState::Active); }
+	FORCEINLINE bool IsPaused() const { return (TaskState == EGameplayTaskState::Paused); }
 	FORCEINLINE bool IsFinished() const { return (TaskState == EGameplayTaskState::Finished); }
 	
+	UGameplayTask* GetChildTask() const { return ChildTask; }
 	IGameplayTaskOwnerInterface* GetTaskOwner() const { return TaskOwner.IsValid() ? &(*TaskOwner) : nullptr; }
 	UGameplayTasksComponent* GetGameplayTasksComponent() { return TasksComponent.Get(); }
 	UGameplayTasksComponent* GetGameplayTasksComponent() const { return TasksComponent.Get(); }
@@ -267,7 +270,7 @@ protected:
 	 *	IMPORTANT! When overriding this function make sure to call Super::OnDestroy(bOwnerFinished) as the last thing,
 	 *		since the function internally marks the task as "Pending Kill", and this may interfere with internal BP mechanics
 	 */
-	virtual void OnDestroy(bool bOwnerFinished);
+	virtual void OnDestroy(bool bInOwnerFinished);
 
 	static IGameplayTaskOwnerInterface* ConvertToTaskOwner(UObject& OwnerObject);
 	static IGameplayTaskOwnerInterface* ConvertToTaskOwner(AActor& OwnerActor);
@@ -276,6 +279,15 @@ protected:
 	virtual void Pause();
 	virtual void Resume();
 
+	// IGameplayTaskOwnerInterface BEGIN
+	virtual UGameplayTasksComponent* GetGameplayTasksComponent(const UGameplayTask& Task) const override;
+	virtual AActor* GetGameplayTaskOwner(const UGameplayTask* Task) const override;
+	virtual AActor* GetGameplayTaskAvatar(const UGameplayTask* Task) const override;
+	virtual uint8 GetGameplayTaskDefaultPriority() const override;
+	virtual void OnGameplayTaskInitialized(UGameplayTask& Task) override;
+	virtual void OnGameplayTaskDeactivated(UGameplayTask& Task) override;
+	// IGameplayTaskOwnerInterface END
+
 private:
 	friend UGameplayTasksComponent;
 
@@ -316,6 +328,8 @@ protected:
 
 	uint32 bClaimRequiredResources : 1;
 	
+	uint32 bOwnerFinished : 1;
+
 	/** Abstract "resource" IDs this task needs available to be able to get activated. */
 	FGameplayResourceSet RequiredResources;
 
@@ -329,6 +343,10 @@ protected:
 
 	TWeakObjectPtr<UGameplayTasksComponent>	TasksComponent;
 
+	/** child task instance */
+	UPROPERTY()
+	UGameplayTask* ChildTask;
+
 #if !(UE_BUILD_SHIPPING || UE_BUILD_TEST)
 	mutable FString DebugDescription;
 public:
@@ -357,7 +375,7 @@ T* UGameplayTask::NewTask(IGameplayTaskOwnerInterface& TaskOwner, FName Instance
 {
 	T* MyObj = NewObject<T>();
 	MyObj->InstanceName = InstanceName;
-	MyObj->InitTask(TaskOwner, TaskOwner.GetDefaultPriority());
+	MyObj->InitTask(TaskOwner, TaskOwner.GetGameplayTaskDefaultPriority());
 	return MyObj;
 }
 
diff --git a/Engine/Source/Runtime/GameplayTasks/Classes/GameplayTaskOwnerInterface.h b/Engine/Source/Runtime/GameplayTasks/Classes/GameplayTaskOwnerInterface.h
--- a/Engine/Source/Runtime/GameplayTasks/Classes/GameplayTaskOwnerInterface.h
+++ b/Engine/Source/Runtime/GameplayTasks/Classes/GameplayTaskOwnerInterface.h
@@ -22,13 +22,25 @@ class GAMEPLAYTASKS_API IGameplayTaskOwnerInterface
 {
 	GENERATED_BODY()
 public:
-	virtual void OnTaskInitialized(UGameplayTask& Task);
-	virtual UGameplayTasksComponent* GetGameplayTasksComponent(const UGameplayTask& Task) const = 0;
-	/** this gets called both when task starts and when task gets resumed. Check Task.GetStatus() if you want to differenciate */
-	virtual void OnTaskActivated(UGameplayTask& Task) = 0;
-	/** this gets called both when task finished and when task gets paused. Check Task.GetStatus() if you want to differenciate */
-	virtual void OnTaskDeactivated(UGameplayTask& Task) = 0;
-	virtual AActor* GetOwnerActor(const UGameplayTask* Task) const = 0;
-	virtual AActor* GetAvatarActor(const UGameplayTask* Task) const;
-	virtual uint8 GetDefaultPriority() const { return FGameplayTasks::DefaultPriority; }
-};
\ No newline at end of file
+
+	/** Finds tasks component for given GameplayTask, Task.GetGameplayTasksComponent() may not be initialized at this point! */
+	virtual UGameplayTasksComponent* GetGameplayTasksComponent(const UGameplayTask& Task) const PURE_VIRTUAL(IGameplayTaskOwnerInterface::GetGameplayTasksComponent, return nullptr;);
+
+	/** Get owner of a task or default one when task is null */
+	virtual AActor* GetGameplayTaskOwner(const UGameplayTask* Task) const PURE_VIRTUAL(IGameplayTaskOwnerInterface::GetGameplayTaskOwner, return nullptr;);
+
+	/** Get "body" of task's owner / default, having location in world (e.g. Owner = AIController, Avatar = Pawn) */
+	virtual AActor* GetGameplayTaskAvatar(const UGameplayTask* Task) const { return GetGameplayTaskOwner(Task); }
+
+	/** Get default priority for running a task */
+	virtual uint8 GetGameplayTaskDefaultPriority() const { return FGameplayTasks::DefaultPriority; }
+
+	/** Notify called after GameplayTask finishes initialization (not active yet) */
+	virtual void OnGameplayTaskInitialized(UGameplayTask& Task) {}
+	
+	/** Notify called after GameplayTask changes state to Active (initial activation or resuming) */
+	virtual void OnGameplayTaskActivated(UGameplayTask& Task) {}
+
+	/** Notify called after GameplayTask changes state from Active (finishing or pausing) */
+	virtual void OnGameplayTaskDeactivated(UGameplayTask& Task) {}
+};
diff --git a/Engine/Source/Runtime/GameplayTasks/Classes/GameplayTaskResource.h b/Engine/Source/Runtime/GameplayTasks/Classes/GameplayTaskResource.h
--- a/Engine/Source/Runtime/GameplayTasks/Classes/GameplayTaskResource.h
+++ b/Engine/Source/Runtime/GameplayTasks/Classes/GameplayTaskResource.h
@@ -49,5 +49,16 @@ protected:
 #endif // WITH_EDITOR
 
 	void UpdateAutoResourceID();
-};
 
+#if !(UE_BUILD_SHIPPING || UE_BUILD_TEST)
+protected:
+	static TArray<FString> ResourceDescriptions;
+	virtual FString GenerateDebugDescription() const;
+
+public:
+	static FString GetDebugDescription(uint8 ResourceId)
+	{
+		return ResourceDescriptions.IsValidIndex(ResourceId) ? ResourceDescriptions[ResourceId] : FString();
+	}
+#endif // !(UE_BUILD_SHIPPING || UE_BUILD_TEST)
+};
diff --git a/Engine/Source/Runtime/GameplayTasks/Classes/GameplayTasksComponent.h b/Engine/Source/Runtime/GameplayTasks/Classes/GameplayTasksComponent.h
--- a/Engine/Source/Runtime/GameplayTasks/Classes/GameplayTasksComponent.h
+++ b/Engine/Source/Runtime/GameplayTasks/Classes/GameplayTasksComponent.h
@@ -44,6 +44,8 @@ struct FGameplayTaskEventData
 
 DECLARE_DYNAMIC_MULTICAST_DELEGATE_TwoParams(FOnClaimedResourcesChangeSignature, FGameplayResourceSet, NewlyClaimed, FGameplayResourceSet, FreshlyReleased);
 
+typedef TArray<UGameplayTask*>::TConstIterator FConstGameplayTaskIterator;
+
 /**
 *	The core ActorComponent for interfacing with the GameplayAbilities System
 */
@@ -95,15 +97,6 @@ public:
 	*	activity into consideration*/
 	virtual bool GetShouldTick() const;
 	
-	virtual UGameplayTasksComponent* GetGameplayTasksComponent(const UGameplayTask& Task) const override { return const_cast<UGameplayTasksComponent*>(this); }
-	/** Adds a task (most often an UAbilityTask) to the list of tasks to be ticked */
-	virtual void OnTaskActivated(UGameplayTask& Task) override;
-	/** Removes a task (most often an UAbilityTask) task from the list of tasks to be ticked */
-	virtual void OnTaskDeactivated(UGameplayTask& Task) override;
-
-	virtual AActor* GetOwnerActor(const UGameplayTask* Task) const override { return GetOwner(); }
-	virtual AActor* GetAvatarActor(const UGameplayTask* Task) const override;
-		
 	/** processes the task and figures out if it should get triggered instantly or wait
 	 *	based on task's RequiredResources, Priority and ResourceOverlapPolicy */
 	void AddTaskReadyForActivation(UGameplayTask& NewTask);
@@ -111,8 +104,18 @@ public:
 	void RemoveResourceConsumingTask(UGameplayTask& Task);
 	void EndAllResourceConsumingTasksOwnedBy(const IGameplayTaskOwnerInterface& TaskOwner);
 
+	bool FindAllTasksOwnedBy(const IGameplayTaskOwnerInterface& TaskOwner, TArray<UGameplayTask*>& FoundTasks) const;
+
 	FORCEINLINE FGameplayResourceSet GetCurrentlyUsedResources() const { return CurrentlyClaimedResources; }
 
+	// BEGIN IGameplayTaskOwnerInterface
+	virtual UGameplayTasksComponent* GetGameplayTasksComponent(const UGameplayTask& Task) const { return const_cast<UGameplayTasksComponent*>(this); }
+	virtual AActor* GetGameplayTaskOwner(const UGameplayTask* Task) const override { return GetOwner(); }
+	virtual AActor* GetGameplayTaskAvatar(const UGameplayTask* Task) const override { return GetOwner(); }
+	virtual void OnGameplayTaskActivated(UGameplayTask& Task) override;
+	virtual void OnGameplayTaskDeactivated(UGameplayTask& Task) override;
+	// END IGameplayTaskOwnerInterface
+
 	UFUNCTION(BlueprintCallable, DisplayName="Run Gameplay Task", Category = "Gameplay Tasks", meta = (AutoCreateRefTerm = "AdditionalRequiredResources, AdditionalClaimedResources", AdvancedDisplay = "AdditionalRequiredResources, AdditionalClaimedResources"))
 	static EGameplayTaskRunResult K2_RunGameplayTask(TScriptInterface<IGameplayTaskOwnerInterface> TaskOwner, UGameplayTask* Task, uint8 Priority, TArray<TSubclassOf<UGameplayTaskResource> > AdditionalRequiredResources, TArray<TSubclassOf<UGameplayTaskResource> > AdditionalClaimedResources);
 
@@ -122,15 +125,27 @@ public:
 	FString GetTickingTasksDescription() const;
 	FString GetTasksPriorityQueueDescription() const;
 #endif // !(UE_BUILD_SHIPPING || UE_BUILD_TEST)
+	FConstGameplayTaskIterator GetTickingTaskIterator() const;
+	FConstGameplayTaskIterator GetPriorityQueueIterator() const;
+
 #if ENABLE_VISUAL_LOG
 	static FString GetTaskStateName(EGameplayTaskState Value);
 	void DescribeSelfToVisLog(struct FVisualLogEntry* Snapshot) const;
 #endif // ENABLE_VISUAL_LOG
 
 protected:
+	struct FEventLock
+	{
+		FEventLock(UGameplayTasksComponent* InOwner);
+		~FEventLock();
+
+	protected:
+		UGameplayTasksComponent* Owner;
+	};
+
 	void RequestTicking();
 	void ProcessTaskEvents();
-	void UpdateTaskActivationFromIndex(int32 StartingIndex, FGameplayResourceSet ResourcesClaimedUpToIndex, FGameplayResourceSet ResourcesBlockedUpToIndex);
+	void UpdateTaskActivations();
 
 	void SetCurrentlyClaimedResources(FGameplayResourceSet NewClaimedSet);
 
@@ -140,4 +155,10 @@ private:
 
 	void AddTaskToPriorityQueue(UGameplayTask& NewTask);
 	void RemoveTaskFromPriorityQueue(UGameplayTask& Task);
-};
\ No newline at end of file
+
+	friend struct FEventLock;
+	int32 EventLockCounter;
+	uint32 bInEventProcessingInProgress : 1;
+
+	FORCEINLINE bool CanProcessEvents() const { return !bInEventProcessingInProgress && (EventLockCounter == 0); }
+};
diff --git a/Engine/Source/Runtime/GameplayTasks/Private/GameplayTask.cpp b/Engine/Source/Runtime/GameplayTasks/Private/GameplayTask.cpp
--- a/Engine/Source/Runtime/GameplayTasks/Private/GameplayTask.cpp
+++ b/Engine/Source/Runtime/GameplayTasks/Private/GameplayTask.cpp
@@ -13,6 +13,7 @@ UGameplayTask::UGameplayTask(const FObjectInitializer& ObjectInitializer)
 	bIsSimulating = false;
 	bOwnedByTasksComponent = false;
 	bClaimRequiredResources = true;
+	bOwnerFinished = false;
 	TaskState = EGameplayTaskState::Uninitialized;
 	ResourceOverlapPolicy = ETaskResourceOverlapPolicy::StartOnTop;
 	Priority = FGameplayTasks::DefaultPriority;
@@ -69,11 +70,6 @@ void UGameplayTask::InitTask(IGameplayTaskOwnerInterface& InTaskOwner, uint8 InP
 {
 	Priority = InPriority;
 	TaskOwner = InTaskOwner;
-	UGameplayTasksComponent* GTComponent = InTaskOwner.GetGameplayTasksComponent(*this);
-	TasksComponent = GTComponent;
-
-	bOwnedByTasksComponent = (TaskOwner == GTComponent);
-	
 	TaskState = EGameplayTaskState::AwaitingActivation;
 
 	if (bClaimRequiredResources)
@@ -81,10 +77,17 @@ void UGameplayTask::InitTask(IGameplayTaskOwnerInterface& InTaskOwner, uint8 InP
 		ClaimedResources.AddSet(RequiredResources);
 	}
 
-	InTaskOwner.OnTaskInitialized(*this);
-	if (bOwnedByTasksComponent == false && GTComponent != nullptr)
+	// call owner.OnGameplayTaskInitialized before accessing owner.GetGameplayTasksComponent, this is required for child tasks
+	InTaskOwner.OnGameplayTaskInitialized(*this);
+
+	UGameplayTasksComponent* GTComponent = InTaskOwner.GetGameplayTasksComponent(*this);
+	TasksComponent = GTComponent;
+	bOwnedByTasksComponent = (TaskOwner == GTComponent);
+
+	// make sure that task component knows about new task
+	if (GTComponent && !bOwnedByTasksComponent)
 	{
-		GTComponent->OnTaskInitialized(*this);
+		GTComponent->OnGameplayTaskInitialized(*this);
 	}
 }
 
@@ -108,11 +111,11 @@ AActor* UGameplayTask::GetOwnerActor() const
 {
 	if (TaskOwner.IsValid())
 	{
-		return TaskOwner->GetOwnerActor(this);		
+		return TaskOwner->GetGameplayTaskOwner(this);		
 	}
 	else if (TasksComponent.IsValid())
 	{
-		return TasksComponent->GetOwnerActor(this);
+		return TasksComponent->GetGameplayTaskOwner(this);
 	}
 
 	return nullptr;
@@ -122,11 +125,11 @@ AActor* UGameplayTask::GetAvatarActor() const
 {
 	if (TaskOwner.IsValid())
 	{
-		return TaskOwner->GetAvatarActor(this);
+		return TaskOwner->GetGameplayTaskAvatar(this);
 	}
 	else if (TasksComponent.IsValid())
 	{
-		return TasksComponent->GetAvatarActor(this);
+		return TasksComponent->GetGameplayTaskAvatar(this);
 	}
 
 	return nullptr;
@@ -140,6 +143,7 @@ void UGameplayTask::TaskOwnerEnded()
 
 	if (TaskState != EGameplayTaskState::Finished && !IsPendingKill())
 	{
+		bOwnerFinished = true;
 		OnDestroy(true);
 	}
 }
@@ -177,22 +181,14 @@ void UGameplayTask::ExternalCancel()
 	EndTask();
 }
 
-void UGameplayTask::OnDestroy(bool bOwnerFinished)
+void UGameplayTask::OnDestroy(bool bInOwnerFinished)
 {
 	ensure(TaskState != EGameplayTaskState::Finished && !IsPendingKill());
-
 	TaskState = EGameplayTaskState::Finished;
 
-	// First of all notify the TaskComponent
 	if (TasksComponent.IsValid())
 	{
-		TasksComponent->OnTaskDeactivated(*this);
-	}
-
-	// Remove ourselves from the owner's task list, if the owner isn't ending
-	if (bOwnedByTasksComponent == false && bOwnerFinished == false && TaskOwner.IsValid() == true)
-	{
-		TaskOwner->OnTaskDeactivated(*this);
+		TasksComponent->OnGameplayTaskDeactivated(*this);
 	}
 
 	MarkPendingKill();
@@ -264,12 +260,7 @@ void UGameplayTask::PerformActivation()
 
 	Activate();
 
-	TasksComponent->OnTaskActivated(*this);
-
-	if (bOwnedByTasksComponent == false && TaskOwner.IsValid())
-	{
-		TaskOwner->OnTaskActivated(*this);
-	}
+	TasksComponent->OnGameplayTaskActivated(*this);
 }
 
 void UGameplayTask::Activate()
@@ -287,12 +278,7 @@ void UGameplayTask::Pause()
 
 	TaskState = EGameplayTaskState::Paused;
 
-	TasksComponent->OnTaskDeactivated(*this);
-
-	if (bOwnedByTasksComponent == false && TaskOwner.IsValid())
-	{
-		TaskOwner->OnTaskDeactivated(*this);
-	}
+	TasksComponent->OnGameplayTaskDeactivated(*this);
 }
 
 void UGameplayTask::Resume()
@@ -303,12 +289,7 @@ void UGameplayTask::Resume()
 
 	TaskState = EGameplayTaskState::Active;
 
-	TasksComponent->OnTaskActivated(*this);
-
-	if (bOwnedByTasksComponent == false && TaskOwner.IsValid())
-	{
-		TaskOwner->OnTaskActivated(*this);
-	}
+	TasksComponent->OnGameplayTaskActivated(*this);
 }
 
 //----------------------------------------------------------------------//
@@ -380,9 +361,17 @@ void UGameplayTask::PauseInTaskQueue()
 //----------------------------------------------------------------------//
 FString UGameplayTask::GenerateDebugDescription() const
 {
-	return RequiresPriorityOrResourceManagement() == false ? GetName()
-		: FString::Printf(TEXT("%s: P:%d %s")
-			, *GetName(), int32(Priority), *RequiredResources.GetDebugDescription());
+	if (RequiresPriorityOrResourceManagement())
+	{
+		UObject* OwnerOb = Cast<UObject>(GetTaskOwner());
+		return FString::Printf(TEXT("%s:%s Pri:%d Owner:%s Res:%s"),
+			*GetName(), InstanceName != NAME_None ? *InstanceName.ToString() : TEXT("-"),
+			(int32)Priority,
+			*GetNameSafe(OwnerOb),
+			*RequiredResources.GetDebugDescription());
+	}
+
+	return GetName();
 }
 
 FString UGameplayTask::GetTaskStateName() const
@@ -392,4 +381,51 @@ FString UGameplayTask::GetTaskStateName() const
 	return Enum->GetEnumName(int32(TaskState));
 }
 
-#endif // !(UE_BUILD_SHIPPING || UE_BUILD_TEST)
\ No newline at end of file
+#endif // !(UE_BUILD_SHIPPING || UE_BUILD_TEST)
+
+//////////////////////////////////////////////////////////////////////////
+// Child tasks
+
+UGameplayTasksComponent* UGameplayTask::GetGameplayTasksComponent(const UGameplayTask& Task) const
+{
+	return ((&Task == ChildTask) || (&Task == this)) ? GetGameplayTasksComponent() : nullptr;
+}
+
+AActor* UGameplayTask::GetGameplayTaskOwner(const UGameplayTask* Task) const
+{
+	return ((Task == ChildTask) || (Task == this)) ? UGameplayTask::GetOwnerActor() : nullptr;
+}
+
+AActor* UGameplayTask::GetGameplayTaskAvatar(const UGameplayTask* Task) const
+{
+	return ((Task == ChildTask) || (Task == this)) ? UGameplayTask::GetAvatarActor() : nullptr;
+}
+
+uint8 UGameplayTask::GetGameplayTaskDefaultPriority() const
+{
+	return GetPriority();
+}
+
+void UGameplayTask::OnGameplayTaskDeactivated(UGameplayTask& Task)
+{
+	// cleanup after deactivation
+	if (&Task == ChildTask)
+	{
+		UE_VLOG(GetGameplayTasksComponent(), LogGameplayTasks, Verbose, TEXT("%s> Child task deactivated: %s"), *GetName(), *Task.GetName());
+		ChildTask = nullptr;
+	}
+}
+
+void UGameplayTask::OnGameplayTaskInitialized(UGameplayTask& Task)
+{
+	UE_VLOG(GetGameplayTasksComponent(), LogGameplayTasks, Verbose, TEXT("%s> Child task initialized: %s"), *GetName(), *Task.GetName());
+
+	// only one child task is allowed
+	if (ChildTask)
+	{
+		UE_VLOG(GetGameplayTasksComponent(), LogGameplayTasks, Verbose, TEXT(">> terminating previous child task: %s"), *ChildTask->GetName());
+		ChildTask->EndTask();
+	}
+
+	ChildTask = &Task;
+}
diff --git a/Engine/Source/Runtime/GameplayTasks/Private/GameplayTaskOwnerInterface.cpp b/Engine/Source/Runtime/GameplayTasks/Private/GameplayTaskOwnerInterface.cpp
--- a/Engine/Source/Runtime/GameplayTasks/Private/GameplayTaskOwnerInterface.cpp
+++ b/Engine/Source/Runtime/GameplayTasks/Private/GameplayTaskOwnerInterface.cpp
@@ -2,13 +2,3 @@
 
 #include "GameplayTasksPrivatePCH.h"
 #include "GameplayTaskOwnerInterface.h"
-
-void IGameplayTaskOwnerInterface::OnTaskInitialized(UGameplayTask& Task)
-{
-
-}
-
-AActor* IGameplayTaskOwnerInterface::GetAvatarActor(const UGameplayTask* Task) const
-{ 
-	return GetOwnerActor(Task);
-}
diff --git a/Engine/Source/Runtime/GameplayTasks/Private/GameplayTaskResource.cpp b/Engine/Source/Runtime/GameplayTasks/Private/GameplayTaskResource.cpp
--- a/Engine/Source/Runtime/GameplayTasks/Private/GameplayTaskResource.cpp
+++ b/Engine/Source/Runtime/GameplayTasks/Private/GameplayTaskResource.cpp
@@ -3,6 +3,10 @@
 #include "GameplayTasksPrivatePCH.h"
 #include "GameplayTaskResource.h"
 
+#if !(UE_BUILD_SHIPPING || UE_BUILD_TEST)
+TArray<FString> UGameplayTaskResource::ResourceDescriptions;
+#endif // !(UE_BUILD_SHIPPING || UE_BUILD_TEST)
+
 UGameplayTaskResource::UGameplayTaskResource(const FObjectInitializer& ObjectInitializer)
 	: Super(ObjectInitializer)
 {
@@ -25,6 +29,12 @@ void UGameplayTaskResource::PostInitProperties()
 		{
 			UpdateAutoResourceID();
 		}
+
+#if !(UE_BUILD_SHIPPING || UE_BUILD_TEST)
+		const uint8 DebugId = GetResourceID();
+		ResourceDescriptions.SetNum(FMath::Max(DebugId + 1, ResourceDescriptions.Num()));
+		ResourceDescriptions[DebugId] = GenerateDebugDescription();
+#endif // !(UE_BUILD_SHIPPING || UE_BUILD_TEST)
 	}
 }
 
@@ -64,4 +74,19 @@ void UGameplayTaskResource::PostEditChangeProperty(FPropertyChangedEvent& Proper
 		}
 	}
 }
-#endif // WITH_EDITOR
\ No newline at end of file
+#endif // WITH_EDITOR
+
+#if !(UE_BUILD_SHIPPING || UE_BUILD_TEST)
+FString UGameplayTaskResource::GenerateDebugDescription() const
+{
+	const FString ClassName = GetClass()->GetName();
+	int32 SeparatorIdx = INDEX_NONE;
+	if (ClassName.FindChar(TEXT('_'), SeparatorIdx))
+	{
+		return ClassName.Mid(SeparatorIdx + 1);
+	}
+
+	return ClassName;
+
+}
+#endif // !(UE_BUILD_SHIPPING || UE_BUILD_TEST)
diff --git a/Engine/Source/Runtime/GameplayTasks/Private/GameplayTasksComponent.cpp b/Engine/Source/Runtime/GameplayTasks/Private/GameplayTasksComponent.cpp
--- a/Engine/Source/Runtime/GameplayTasks/Private/GameplayTasksComponent.cpp
+++ b/Engine/Source/Runtime/GameplayTasks/Private/GameplayTasksComponent.cpp
@@ -2,6 +2,7 @@
 
 #include "GameplayTasksPrivatePCH.h"
 #include "GameplayTasksComponent.h"
+#include "GameplayTaskResource.h"
 #include "GameplayTask.h"
 #include "MessageLog.h"
 
@@ -18,6 +19,27 @@ namespace
 	}
 }
 
+UGameplayTasksComponent::FEventLock::FEventLock(UGameplayTasksComponent* InOwner) : Owner(InOwner)
+{
+	if (Owner)
+	{
+		Owner->EventLockCounter++;
+	}
+}
+
+UGameplayTasksComponent::FEventLock::~FEventLock()
+{
+	if (Owner)
+	{
+		Owner->EventLockCounter--;
+
+		if (Owner->TaskEvents.Num() && Owner->CanProcessEvents())
+		{
+			Owner->ProcessTaskEvents();
+		}
+	}
+}
+
 UGameplayTasksComponent::UGameplayTasksComponent(const FObjectInitializer& ObjectInitializer)
 	: Super(ObjectInitializer)
 {
@@ -26,11 +48,15 @@ UGameplayTasksComponent::UGameplayTasksComponent(const FObjectInitializer& Objec
 	PrimaryComponentTick.bCanEverTick = true;
 
 	bReplicates = true;
+	bInEventProcessingInProgress = false;
 	TopActivePriority = 0;
 }
 	
-void UGameplayTasksComponent::OnTaskActivated(UGameplayTask& Task)
+void UGameplayTasksComponent::OnGameplayTaskActivated(UGameplayTask& Task)
 {
+	// process events after finishing all operations
+	FEventLock ScopeEventLock(this);
+
 	if (Task.IsTickingTask())
 	{
 		check(TickingTasks.Contains(&Task) == false);
@@ -47,10 +73,32 @@ void UGameplayTasksComponent::OnTaskActivated(UGameplayTask& Task)
 		check(SimulatedTasks.Contains(&Task) == false);
 		SimulatedTasks.Add(&Task);
 	}
+
+	IGameplayTaskOwnerInterface* TaskOwner = Task.GetTaskOwner();
+	if (!Task.IsOwnedByTasksComponent() && TaskOwner)
+	{
+		TaskOwner->OnGameplayTaskActivated(Task);
+	}
 }
 
-void UGameplayTasksComponent::OnTaskDeactivated(UGameplayTask& Task)
+void UGameplayTasksComponent::OnGameplayTaskDeactivated(UGameplayTask& Task)
 {
+	// process events after finishing all operations
+	FEventLock ScopeEventLock(this);
+	const bool bIsFinished = (Task.GetState() == EGameplayTaskState::Finished);
+
+	if (Task.GetChildTask() && bIsFinished)
+	{
+		if (Task.HasOwnerFinished())
+		{
+			Task.GetChildTask()->TaskOwnerEnded();
+		}
+		else
+		{
+			Task.GetChildTask()->EndTask();
+		}
+	}
+
 	if (Task.IsTickingTask())
 	{
 		// If we are removing our last ticking task, set this component as inactive so it stops ticking
@@ -62,13 +110,19 @@ void UGameplayTasksComponent::OnTaskDeactivated(UGameplayTask& Task)
 		SimulatedTasks.RemoveSingleSwap(&Task);
 	}
 
-	UpdateShouldTick();
-
 	// Resource-using task
-	if (Task.RequiresPriorityOrResourceManagement() == true && Task.GetState() == EGameplayTaskState::Finished)
+	if (Task.RequiresPriorityOrResourceManagement() && bIsFinished)
 	{
 		OnTaskEnded(Task);
 	}
+
+	IGameplayTaskOwnerInterface* TaskOwner = Task.GetTaskOwner();
+	if (!Task.IsOwnedByTasksComponent() && !Task.HasOwnerFinished() && TaskOwner)
+	{
+		TaskOwner->OnGameplayTaskDeactivated(Task);
+	}
+
+	UpdateShouldTick();
 }
 
 void UGameplayTasksComponent::OnTaskEnded(UGameplayTask& Task)
@@ -191,11 +245,6 @@ void UGameplayTasksComponent::UpdateShouldTick()
 	}
 }
 
-AActor* UGameplayTasksComponent::GetAvatarActor(const UGameplayTask* Task) const
-{
-	return GetOwner();
-}
-
 //----------------------------------------------------------------------//
 // Priority and resources handling
 //----------------------------------------------------------------------//
@@ -207,7 +256,7 @@ void UGameplayTasksComponent::AddTaskReadyForActivation(UGameplayTask& NewTask)
 	
 	TaskEvents.Add(FGameplayTaskEventData(EGameplayTaskEvent::Add, NewTask));
 	// trigger the actual processing only if it was the first event added to the list
-	if (TaskEvents.Num() == 1)
+	if (TaskEvents.Num() == 1 && CanProcessEvents())
 	{
 		ProcessTaskEvents();
 	}
@@ -219,7 +268,7 @@ void UGameplayTasksComponent::RemoveResourceConsumingTask(UGameplayTask& Task)
 
 	TaskEvents.Add(FGameplayTaskEventData(EGameplayTaskEvent::Remove, Task));
 	// trigger the actual processing only if it was the first event added to the list
-	if (TaskEvents.Num() == 1)
+	if (TaskEvents.Num() == 1 && CanProcessEvents())
 	{
 		ProcessTaskEvents();
 	}
@@ -227,128 +276,122 @@ void UGameplayTasksComponent::RemoveResourceConsumingTask(UGameplayTask& Task)
 
 void UGameplayTasksComponent::EndAllResourceConsumingTasksOwnedBy(const IGameplayTaskOwnerInterface& TaskOwner)// , bool bNotifyOwner)
 {
-	for (int32 TaskIndex = TaskPriorityQueue.Num() - 1; TaskIndex >= 0; --TaskIndex)
 	{
-		if (TaskPriorityQueue[TaskIndex] == nullptr)
+		FEventLock ScopeEventLock(this);
+
+		// can't reuse FindAllTasksOwnedBy here, it's important to remove from queue immediately
+		for (int32 Idx = TaskPriorityQueue.Num() - 1; Idx >= 0; Idx--)
 		{
-			TaskPriorityQueue.RemoveAt(TaskIndex, 1, /*bAllowShrinking=*/false);
+			if (TaskPriorityQueue[Idx] && TaskPriorityQueue[Idx]->GetTaskOwner() == &TaskOwner)
+			{
+				UGameplayTask* Task = TaskPriorityQueue[Idx];
+				TaskPriorityQueue[Idx] = nullptr;
+				Task->TaskOwnerEnded();
+			}
 		}
-		else if (TaskPriorityQueue[TaskIndex]->GetTaskOwner() == &TaskOwner)
+	}
+
+	UpdateTaskActivations();
+}
+
+bool UGameplayTasksComponent::FindAllTasksOwnedBy(const IGameplayTaskOwnerInterface& TaskOwner, TArray<UGameplayTask*>& FoundTasks) const
+{
+	int32 NumFound = 0;
+	for (int32 Idx = 0; Idx < TaskPriorityQueue.Num(); Idx++)
+	{
+		if (TaskPriorityQueue[Idx] && TaskPriorityQueue[Idx]->GetTaskOwner() == &TaskOwner)
 		{
-			UGameplayTask* Task = TaskPriorityQueue[TaskIndex];
-			TaskPriorityQueue.RemoveAt(TaskIndex, 1, /*bAllowShrinking=*/false);
-			Task->TaskOwnerEnded();
+			FoundTasks.Add(TaskPriorityQueue[Idx]);
+			NumFound++;
 		}
 	}
 
-	UpdateTaskActivationFromIndex(0, FGameplayResourceSet(), FGameplayResourceSet());
+	return (NumFound > 0);
 }
 
 void UGameplayTasksComponent::ProcessTaskEvents()
 {
-	static const int32 ErronousIterationsLimit = 100;
+	static const int32 MaxIterations = 16;
+	bInEventProcessingInProgress = true;
 
-	// note that this function allows called functions to add new elements to 
-	// TaskEvents array that the main loop is iterating over. It's a feature
-	for (int32 EventIndex = 0; EventIndex < TaskEvents.Num(); ++EventIndex)
+	int32 IterCounter = 0;
+	while (TaskEvents.Num() > 0)
 	{
-		UE_VLOG(this, LogGameplayTasks, Verbose, TEXT("UGameplayTasksComponent::ProcessTaskEvents: %s event %s")
-			, *TaskEvents[EventIndex].RelatedTask.GetName(), GetGameplayTaskEventName(TaskEvents[EventIndex].Event));
-
-		if (TaskEvents[EventIndex].RelatedTask.IsPendingKill())
+		IterCounter++;
+		if (IterCounter > MaxIterations)
 		{
-			UE_VLOG(this, LogGameplayTasks, Verbose, TEXT("%s is PendingKill"), *TaskEvents[EventIndex].RelatedTask.GetName());
-			// we should ignore it, but just in case run the removal code.
-			RemoveTaskFromPriorityQueue(TaskEvents[EventIndex].RelatedTask);
-			continue;
+			UE_VLOG(this, LogGameplayTasks, Error, TEXT("UGameplayTasksComponent::ProcessTaskEvents has exceeded allowes number of iterations. Check your GameplayTasks for logic loops!"));
+			break;
 		}
 
-		switch (TaskEvents[EventIndex].Event)
+		for (int32 EventIndex = 0; EventIndex < TaskEvents.Num(); ++EventIndex)
 		{
-		case EGameplayTaskEvent::Add:
-			if (TaskEvents[EventIndex].RelatedTask.TaskState != EGameplayTaskState::Finished)
+			UE_VLOG(this, LogGameplayTasks, Verbose, TEXT("UGameplayTasksComponent::ProcessTaskEvents: %s event %s")
+				, *TaskEvents[EventIndex].RelatedTask.GetName(), GetGameplayTaskEventName(TaskEvents[EventIndex].Event));
+
+			if (TaskEvents[EventIndex].RelatedTask.IsPendingKill())
 			{
-				AddTaskToPriorityQueue(TaskEvents[EventIndex].RelatedTask);
+				UE_VLOG(this, LogGameplayTasks, Verbose, TEXT("%s is PendingKill"), *TaskEvents[EventIndex].RelatedTask.GetName());
+				// we should ignore it, but just in case run the removal code.
+				RemoveTaskFromPriorityQueue(TaskEvents[EventIndex].RelatedTask);
+				continue;
 			}
-			else
+
+			switch (TaskEvents[EventIndex].Event)
 			{
-				UE_VLOG(this, LogGameplayTasks, Error, TEXT("UGameplayTasksComponent::ProcessTaskEvents trying to add a finished task to priority queue!"));
+			case EGameplayTaskEvent::Add:
+				if (TaskEvents[EventIndex].RelatedTask.TaskState != EGameplayTaskState::Finished)
+				{
+					AddTaskToPriorityQueue(TaskEvents[EventIndex].RelatedTask);
+				}
+				else
+				{
+					UE_VLOG(this, LogGameplayTasks, Error, TEXT("UGameplayTasksComponent::ProcessTaskEvents trying to add a finished task to priority queue!"));
+				}
+				break;
+			case EGameplayTaskEvent::Remove:
+				RemoveTaskFromPriorityQueue(TaskEvents[EventIndex].RelatedTask);
+				break;
+			default:
+				checkNoEntry();
+				break;
 			}
-			break;
-		case EGameplayTaskEvent::Remove:
-			RemoveTaskFromPriorityQueue(TaskEvents[EventIndex].RelatedTask);
-			break;
-		default:
-			checkNoEntry();
-			break;
 		}
 
-		if (EventIndex >= ErronousIterationsLimit)
-		{
-			UE_VLOG(this, LogGameplayTasks, Error, TEXT("UGameplayTasksComponent::ProcessTaskEvents has exceeded warning-level of iterations. Check your GameplayTasks for logic loops!"));
-		}
+		TaskEvents.Reset();
+		UpdateTaskActivations();
+
+		// task activation changes may create new events, loop over to check it
 	}
 
-	TaskEvents.Reset();
+	bInEventProcessingInProgress = false;
 }
 
 void UGameplayTasksComponent::AddTaskToPriorityQueue(UGameplayTask& NewTask)
 {
-	// The generic idea is as follows:
-	// 1. Find insertion/removal point X
-	//		While looking for it add up all ResourcesUsedByActiveUpToX and ResourcesRequiredUpToX
-	//		ResourcesUsedByActiveUpToX - resources required by ACTIVE tasks
-	//		ResourcesRequiredUpToX - resources required by both active and inactive tasks on the way to X
-	// 2. Insert Task at X
-	// 3. Starting from X proceed down the queue
-	//	a. If ConsideredTask.Resources overlaps with ResourcesRequiredUpToX then PAUSE it
-	//	b. Else, ACTIVATE ConsideredTask and add ConsideredTask.Resources to ResourcesUsedByActiveUpToX
-	//	c. Add ConsideredTask.Resources to ResourcesRequiredUpToX
-	// 4. Set this->CurrentlyUsedResources to ResourcesUsedByActiveUpToX
-	//
-	// Points 3-4 are implemented as a separate function, UpdateTaskActivationFromIndex,
-	// since it's a common code between adding and removing tasks
-	
-	const bool bStartOnTopOfSamePriority = NewTask.GetResourceOverlapPolicy() == ETaskResourceOverlapPolicy::StartOnTop;
-
-	FGameplayResourceSet ResourcesClaimedUpToX;
-	FGameplayResourceSet ResourcesBlockedUpToIndex;
+	const bool bStartOnTopOfSamePriority = (NewTask.GetResourceOverlapPolicy() == ETaskResourceOverlapPolicy::StartOnTop);
 	int32 InsertionPoint = INDEX_NONE;
 	
-	if (TaskPriorityQueue.Num() > 0)
+	for (int32 Idx = 0; Idx < TaskPriorityQueue.Num(); ++Idx)
 	{
-		for (int32 TaskIndex = 0; TaskIndex < TaskPriorityQueue.Num(); ++TaskIndex)
+		if (TaskPriorityQueue[Idx] == nullptr)
 		{
-			ensure(TaskPriorityQueue[TaskIndex]);
-			if (TaskPriorityQueue[TaskIndex] == nullptr)
-			{
-				continue;
-			}
-
-			if ((bStartOnTopOfSamePriority && TaskPriorityQueue[TaskIndex]->GetPriority() <= NewTask.GetPriority())
-				|| (!bStartOnTopOfSamePriority && TaskPriorityQueue[TaskIndex]->GetPriority() < NewTask.GetPriority()))
-			{
-				TaskPriorityQueue.Insert(&NewTask, TaskIndex);
-				InsertionPoint = TaskIndex;
-				break;
-			}
+			continue;
+		}
 
-			const FGameplayResourceSet ClaimedResources = TaskPriorityQueue[TaskIndex]->GetClaimedResources();
-			if (TaskPriorityQueue[TaskIndex]->IsActive())
-			{
-				ResourcesClaimedUpToX.AddSet(ClaimedResources);
-			}
-			ResourcesBlockedUpToIndex.AddSet(ClaimedResources);
+		if ((bStartOnTopOfSamePriority && TaskPriorityQueue[Idx]->GetPriority() <= NewTask.GetPriority())
+			|| (!bStartOnTopOfSamePriority && TaskPriorityQueue[Idx]->GetPriority() < NewTask.GetPriority()))
+		{
+			TaskPriorityQueue.Insert(&NewTask, Idx);
+			InsertionPoint = Idx;
+			break;
 		}
 	}
 	
 	if (InsertionPoint == INDEX_NONE)
 	{
 		TaskPriorityQueue.Add(&NewTask);
-		InsertionPoint = TaskPriorityQueue.Num() - 1;
 	}
-
-	UpdateTaskActivationFromIndex(InsertionPoint, ResourcesClaimedUpToX, ResourcesBlockedUpToIndex);
 }
 
 void UGameplayTasksComponent::RemoveTaskFromPriorityQueue(UGameplayTask& Task)
@@ -356,47 +399,7 @@ void UGameplayTasksComponent::RemoveTaskFromPriorityQueue(UGameplayTask& Task)
 	const int32 RemovedTaskIndex = TaskPriorityQueue.Find(&Task);
 	if (RemovedTaskIndex != INDEX_NONE)
 	{
-		if (TaskPriorityQueue.Num() > 1)
-		{
-			// sum up resources up to TaskIndex
-			FGameplayResourceSet ResourcesClaimedUpToX;
-			FGameplayResourceSet ResourcesBlockedUpToIndex;
-			for (int32 TaskIndex = 0; TaskIndex < RemovedTaskIndex; ++TaskIndex)
-			{
-				ensure(TaskPriorityQueue[TaskIndex]);
-				if (TaskPriorityQueue[TaskIndex] == nullptr)
-				{
-					continue;
-				}
-
-				const FGameplayResourceSet ClaimedResources = TaskPriorityQueue[TaskIndex]->GetClaimedResources();
-				if (TaskPriorityQueue[TaskIndex]->IsActive())
-				{
-					ResourcesClaimedUpToX.AddSet(ClaimedResources);
-				}				
-				ResourcesBlockedUpToIndex.AddSet(ClaimedResources);
-			}
-
-			// don't forget to actually remove the task from the queue
-			TaskPriorityQueue.RemoveAt(RemovedTaskIndex, 1, /*bAllowShrinking=*/false);
-
-			// if it wasn't the last item then proceed as usual
-			if (RemovedTaskIndex < TaskPriorityQueue.Num())
-			{
-				UpdateTaskActivationFromIndex(RemovedTaskIndex, ResourcesClaimedUpToX, ResourcesBlockedUpToIndex);
-			}
-			else
-			{
-				// no need to do extra processing. This was the last task, so 
-				// ResourcesUsedByActiveUpToX is the CurrentlyUsedResources
-				SetCurrentlyClaimedResources(ResourcesClaimedUpToX);
-			}
-		}
-		else
-		{
-			TaskPriorityQueue.Pop(/*bAllowShrinking=*/false);
-			SetCurrentlyClaimedResources(FGameplayResourceSet());
-		}
+		TaskPriorityQueue.RemoveAt(RemovedTaskIndex, 1, /*bAllowShrinking=*/false);
 	}
 	else
 	{
@@ -405,37 +408,59 @@ void UGameplayTasksComponent::RemoveTaskFromPriorityQueue(UGameplayTask& Task)
 	}
 }
 
-void UGameplayTasksComponent::UpdateTaskActivationFromIndex(int32 StartingIndex, FGameplayResourceSet ResourcesClaimedUpToIndex, FGameplayResourceSet ResourcesBlockedUpToIndex)
+void UGameplayTasksComponent::UpdateTaskActivations()
 {
+	FGameplayResourceSet ResourcesClaimed;
+	bool bHasNulls = false;
+
 	if (TaskPriorityQueue.Num() > 0)
 	{
-		check(TaskPriorityQueue.IsValidIndex(StartingIndex));
+		TArray<UGameplayTask*> ActivationList;
+		ActivationList.Reserve(TaskPriorityQueue.Num());
 
-		TArray<UGameplayTask*> TaskPriorityQueueCopy = TaskPriorityQueue;
-		for (int32 TaskIndex = StartingIndex; TaskIndex < TaskPriorityQueueCopy.Num(); ++TaskIndex)
+		FGameplayResourceSet ResourcesBlocked;
+		for (int32 TaskIndex = 0; TaskIndex < TaskPriorityQueue.Num(); ++TaskIndex)
 		{
-			ensure(TaskPriorityQueueCopy[TaskIndex]);
-			if (TaskPriorityQueueCopy[TaskIndex] == nullptr)
+			// allow nulls in queue, can be introduced by EndAllResourceConsumingTasksOwnedBy or custom logic in task pause handler
+			// but we still need to walk through it in order
+			if (TaskPriorityQueue[TaskIndex] == nullptr)
 			{
+				bHasNulls = true;
 				continue;
 			}
 
-			const FGameplayResourceSet RequiredResources = TaskPriorityQueueCopy[TaskIndex]->GetRequiredResources();
-			const FGameplayResourceSet ClaimedResources = TaskPriorityQueueCopy[TaskIndex]->GetClaimedResources();
-			if (RequiredResources.GetOverlap(ResourcesBlockedUpToIndex).IsEmpty())
+			const FGameplayResourceSet RequiredResources = TaskPriorityQueue[TaskIndex]->GetRequiredResources();
+			const FGameplayResourceSet ClaimedResources = TaskPriorityQueue[TaskIndex]->GetClaimedResources();
+			if (RequiredResources.GetOverlap(ResourcesBlocked).IsEmpty())
 			{
-				TaskPriorityQueueCopy[TaskIndex]->ActivateInTaskQueue();
-				ResourcesClaimedUpToIndex.AddSet(ClaimedResources);
+				// postpone activations, it's some tasks (like MoveTo) require pausing old ones first
+				ActivationList.Add(TaskPriorityQueue[TaskIndex]);
+				ResourcesClaimed.AddSet(ClaimedResources);
 			}
 			else
 			{
-				TaskPriorityQueueCopy[TaskIndex]->PauseInTaskQueue();
+				TaskPriorityQueue[TaskIndex]->PauseInTaskQueue();
+			}
+
+			ResourcesBlocked.AddSet(ClaimedResources);
+		}
+
+		for (int32 Idx = 0; Idx < ActivationList.Num(); Idx++)
+		{
+			if (ActivationList[Idx])
+			{
+				ActivationList[Idx]->ActivateInTaskQueue();
 			}
-			ResourcesBlockedUpToIndex.AddSet(ClaimedResources);
 		}
 	}
 	
-	SetCurrentlyClaimedResources(ResourcesClaimedUpToIndex);
+	SetCurrentlyClaimedResources(ResourcesClaimed);
+
+	// remove all null entries after processing activation changes
+	if (bHasNulls)
+	{
+		TaskPriorityQueue.RemoveAll([](UGameplayTask* Task) { return Task == nullptr; });
+	}
 }
 
 void UGameplayTasksComponent::SetCurrentlyClaimedResources(FGameplayResourceSet NewClaimedSet)
@@ -488,6 +513,16 @@ FString UGameplayTasksComponent::GetTasksPriorityQueueDescription() const
 }
 #endif // !(UE_BUILD_SHIPPING || UE_BUILD_TEST)
 
+FConstGameplayTaskIterator UGameplayTasksComponent::GetTickingTaskIterator() const
+{
+	return TickingTasks.CreateConstIterator();
+}
+
+FConstGameplayTaskIterator UGameplayTasksComponent::GetPriorityQueueIterator() const
+{
+	return TaskPriorityQueue.CreateConstIterator();
+}
+
 #if ENABLE_VISUAL_LOG
 void UGameplayTasksComponent::DescribeSelfToVisLog(FVisualLogEntry* Snapshot) const
 {
@@ -625,10 +660,25 @@ EGameplayTaskRunResult UGameplayTasksComponent::K2_RunGameplayTask(TScriptInterf
 //----------------------------------------------------------------------//
 FString FGameplayResourceSet::GetDebugDescription() const
 {
-	static const int32 FlagsCount = sizeof(FFlagContainer)* 8;
-	TCHAR Description[FlagsCount + 1];
+	static const int32 FlagsCount = sizeof(FFlagContainer) * 8;
 	FFlagContainer FlagsCopy = Flags;
 	int32 FlagIndex = 0;
+
+#if !(UE_BUILD_SHIPPING || UE_BUILD_TEST)
+	FString Description;
+	for (; FlagIndex < FlagsCount && FlagsCopy != 0; ++FlagIndex)
+	{
+		if (FlagsCopy & (1 << FlagIndex))
+		{
+			Description += UGameplayTaskResource::GetDebugDescription(FlagIndex);
+			Description += TEXT(' ');
+		}
+
+		FlagsCopy &= ~(1 << FlagIndex);
+	}
+	return Description;
+#else
+	TCHAR Description[FlagsCount + 1];
 	for (; FlagIndex < FlagsCount && FlagsCopy != 0; ++FlagIndex)
 	{
 		Description[FlagIndex] = (FlagsCopy & (1 << FlagIndex)) ? TCHAR('1') : TCHAR('0');
@@ -636,6 +686,7 @@ FString FGameplayResourceSet::GetDebugDescription() const
 	}
 	Description[FlagIndex] = TCHAR('\0');
 	return FString(Description);
+#endif // !(UE_BUILD_SHIPPING || UE_BUILD_TEST)
 }
 
 #undef LOCTEXT_NAMESPACE
\ No newline at end of file
diff --git a/Engine/Source/Runtime/GameplayTasks/Private/Tasks/GameplayTask_SpawnActor.cpp b/Engine/Source/Runtime/GameplayTasks/Private/Tasks/GameplayTask_SpawnActor.cpp
--- a/Engine/Source/Runtime/GameplayTasks/Private/Tasks/GameplayTask_SpawnActor.cpp
+++ b/Engine/Source/Runtime/GameplayTasks/Private/Tasks/GameplayTask_SpawnActor.cpp
@@ -13,7 +13,7 @@ UGameplayTask_SpawnActor* UGameplayTask_SpawnActor::SpawnActor(TScriptInterface<
 	bool bCanSpawn = true;
 	if (bSpawnOnlyOnAuthority == true)
 	{
-		AActor* TaskOwnerActor = TaskOwner->GetOwnerActor(nullptr);
+		AActor* TaskOwnerActor = TaskOwner->GetGameplayTaskOwner(nullptr);
 		if (TaskOwnerActor)
 		{
 			bCanSpawn = (TaskOwnerActor->Role == ROLE_Authority);
diff --git a/Templates/TP_FirstPerson/Source/TP_FirstPerson/TP_FirstPersonProjectile.cpp b/Templates/TP_FirstPerson/Source/TP_FirstPerson/TP_FirstPersonProjectile.cpp
--- a/Templates/TP_FirstPerson/Source/TP_FirstPerson/TP_FirstPersonProjectile.cpp
+++ b/Templates/TP_FirstPerson/Source/TP_FirstPerson/TP_FirstPersonProjectile.cpp
@@ -31,7 +31,7 @@ ATP_FirstPersonProjectile::ATP_FirstPersonProjectile()
 	InitialLifeSpan = 3.0f;
 }
 
-void ATP_FirstPersonProjectile::OnHit(AActor* OtherActor, UPrimitiveComponent* OtherComp, FVector NormalImpulse, const FHitResult& Hit)
+void ATP_FirstPersonProjectile::OnHit(UPrimitiveComponent* HitComp, AActor* OtherActor, UPrimitiveComponent* OtherComp, FVector NormalImpulse, const FHitResult& Hit)
 {
 	// Only add impulse and destroy projectile if we hit a physics
 	if ((OtherActor != NULL) && (OtherActor != this) && (OtherComp != NULL) && OtherComp->IsSimulatingPhysics())
diff --git a/Templates/TP_FirstPerson/Source/TP_FirstPerson/TP_FirstPersonProjectile.h b/Templates/TP_FirstPerson/Source/TP_FirstPerson/TP_FirstPersonProjectile.h
--- a/Templates/TP_FirstPerson/Source/TP_FirstPerson/TP_FirstPersonProjectile.h
+++ b/Templates/TP_FirstPerson/Source/TP_FirstPerson/TP_FirstPersonProjectile.h
@@ -21,7 +21,7 @@ public:
 
 	/** called when projectile hits something */
 	UFUNCTION()
-	void OnHit(AActor* OtherActor, UPrimitiveComponent* OtherComp, FVector NormalImpulse, const FHitResult& Hit);
+	void OnHit(UPrimitiveComponent* HitComp, AActor* OtherActor, UPrimitiveComponent* OtherComp, FVector NormalImpulse, const FHitResult& Hit);
 
 	/** Returns CollisionComp subobject **/
 	FORCEINLINE class USphereComponent* GetCollisionComp() const { return CollisionComp; }
diff --git a/Templates/TP_TwinStick/Source/TP_TwinStick/TP_TwinStickProjectile.cpp b/Templates/TP_TwinStick/Source/TP_TwinStick/TP_TwinStickProjectile.cpp
--- a/Templates/TP_TwinStick/Source/TP_TwinStick/TP_TwinStickProjectile.cpp
+++ b/Templates/TP_TwinStick/Source/TP_TwinStick/TP_TwinStickProjectile.cpp
@@ -30,7 +30,7 @@ ATP_TwinStickProjectile::ATP_TwinStickProjectile()
 	InitialLifeSpan = 3.0f;
 }
 
-void ATP_TwinStickProjectile::OnHit(AActor* OtherActor, UPrimitiveComponent* OtherComp, FVector NormalImpulse, const FHitResult& Hit)
+void ATP_TwinStickProjectile::OnHit(UPrimitiveComponent* HitComp, AActor* OtherActor, UPrimitiveComponent* OtherComp, FVector NormalImpulse, const FHitResult& Hit)
 {
 	// Only add impulse and destroy projectile if we hit a physics
 	if ((OtherActor != NULL) && (OtherActor != this) && (OtherComp != NULL) && OtherComp->IsSimulatingPhysics())
diff --git a/Templates/TP_TwinStick/Source/TP_TwinStick/TP_TwinStickProjectile.h b/Templates/TP_TwinStick/Source/TP_TwinStick/TP_TwinStickProjectile.h
--- a/Templates/TP_TwinStick/Source/TP_TwinStick/TP_TwinStickProjectile.h
+++ b/Templates/TP_TwinStick/Source/TP_TwinStick/TP_TwinStickProjectile.h
@@ -25,7 +25,7 @@ public:
 
 	/** Function to handle the projectile hitting something */
 	UFUNCTION()
-	void OnHit(AActor* OtherActor, UPrimitiveComponent* OtherComp, FVector NormalImpulse, const FHitResult& Hit);
+	void OnHit(UPrimitiveComponent* HitComp, AActor* OtherActor, UPrimitiveComponent* OtherComp, FVector NormalImpulse, const FHitResult& Hit);
 
 	/** Returns ProjectileMesh subobject **/
 	FORCEINLINE UStaticMeshComponent* GetProjectileMesh() const { return ProjectileMesh; }
