Expanded the information provided by the ensure in UGameplayTask::OnDestroy

diff --git a/Engine/Source/Runtime/GameplayTasks/Private/GameplayTask.cpp b/Engine/Source/Runtime/GameplayTasks/Private/GameplayTask.cpp
--- a/Engine/Source/Runtime/GameplayTasks/Private/GameplayTask.cpp
+++ b/Engine/Source/Runtime/GameplayTasks/Private/GameplayTask.cpp
@@ -201,7 +201,11 @@ void UGameplayTask::ExternalCancel()
 
 void UGameplayTask::OnDestroy(bool bInOwnerFinished)
 {
-	ensure(TaskState != EGameplayTaskState::Finished && IsValid(this));
+	if (ensureMsgf(IsValid(this), TEXT("OnDestroy called on invalid gameplay task")))
+	{
+		const FString OwnerName = TaskOwner.IsValid() ? TaskOwner.GetObject()->GetName() : TEXT("Invalid GameplayTask Owner");
+		ensureMsgf(TaskState != EGameplayTaskState::Finished, TEXT("%s OnDestroy called, current state: %i, owner name: %s"), *GetName(), TaskState, *OwnerName);
+	}
 	TaskState = EGameplayTaskState::Finished;
 
 	if (UGameplayTasksComponent* TasksPtr = TasksComponent.Get())
