GameplayTask: Added new task resource overlap policies to cancel existing tasks.

diff --git a/Engine/Source/Runtime/GameplayTasks/Classes/GameplayTask.h b/Engine/Source/Runtime/GameplayTasks/Classes/GameplayTask.h
--- a/Engine/Source/Runtime/GameplayTasks/Classes/GameplayTask.h
+++ b/Engine/Source/Runtime/GameplayTasks/Classes/GameplayTask.h
@@ -35,6 +35,10 @@ enum class ETaskResourceOverlapPolicy : uint8
 	StartOnTop,
 	/** Wait for other same-priority tasks to finish. */
 	StartAtEnd,
+	/** Request to cancel same or lower priority tasks, if tasks not ended pause overlapping same-priority tasks. */
+	RequestCancelAndStartOnTop,
+	/** Request to cancel same or lower priority tasks, wait for any remaining overlapping same-priority tasks to finish. */
+	RequestCancelAndStartAtEnd,
 };
 	
 USTRUCT(BlueprintType)
diff --git a/Engine/Source/Runtime/GameplayTasks/Private/GameplayTasksComponent.cpp b/Engine/Source/Runtime/GameplayTasks/Private/GameplayTasksComponent.cpp
--- a/Engine/Source/Runtime/GameplayTasks/Private/GameplayTasksComponent.cpp
+++ b/Engine/Source/Runtime/GameplayTasks/Private/GameplayTasksComponent.cpp
@@ -457,7 +457,31 @@
 
 void UGameplayTasksComponent::AddTaskToPriorityQueue(UGameplayTask& NewTask)
 {
-	const bool bStartOnTopOfSamePriority = (NewTask.GetResourceOverlapPolicy() == ETaskResourceOverlapPolicy::StartOnTop);
+	if ((NewTask.GetResourceOverlapPolicy() == ETaskResourceOverlapPolicy::RequestCancelAndStartOnTop)
+		|| (NewTask.GetResourceOverlapPolicy() == ETaskResourceOverlapPolicy::RequestCancelAndStartAtEnd))
+	{
+		const FGameplayResourceSet NewClaimedResources = NewTask.GetClaimedResources();
+		TArray<UGameplayTask*, TInlineAllocator<2>> CancelList;
+
+		for (UGameplayTask* Task : TaskPriorityQueue)
+		{
+			if (Task != nullptr
+				&& Task->GetPriority() <= NewTask.GetPriority()
+				&& Task->GetClaimedResources().HasAnyID(NewClaimedResources))
+			{
+				// Postpone cancelling, as cancel can call EndTask() and may alter the TaskPriorityQueue.  
+				CancelList.Add(Task);
+			}
+		}
+
+		for (UGameplayTask* Task : CancelList)
+		{
+			Task->ExternalCancel();
+		}
+	}
+	
+	const bool bStartOnTopOfSamePriority = (NewTask.GetResourceOverlapPolicy() == ETaskResourceOverlapPolicy::StartOnTop)
+										|| (NewTask.GetResourceOverlapPolicy() == ETaskResourceOverlapPolicy::RequestCancelAndStartOnTop);
 	int32 InsertionPoint = INDEX_NONE;
 	
 	for (int32 Idx = 0; Idx < TaskPriorityQueue.Num(); ++Idx)
