[GameplayTasksComponent]
- now inherits fromIVisualLoggerDebugSnapshotInterface so debug snapshots are taken even if there is no AIController (e.g. replicated pawn)
- no longer tick invalid tasks (i.e. destroyed by the replication layer but still part of the TickingTasks)
- added GetSimulatedTaskIterator

diff --git a/Engine/Source/Runtime/AIModule/Private/AIController.cpp b/Engine/Source/Runtime/AIModule/Private/AIController.cpp
--- a/Engine/Source/Runtime/AIModule/Private/AIController.cpp
+++ b/Engine/Source/Runtime/AIModule/Private/AIController.cpp
@@ -164,11 +164,6 @@ void AAIController::GrabDebugSnapshot(FVisualLogEntry* Snapshot) const
 	{
 		PerceptionComponent->DescribeSelfToVisLog(Snapshot);
 	}
-
-	if (CachedGameplayTasksComponent != nullptr)
-	{
-		CachedGameplayTasksComponent->DescribeSelfToVisLog(Snapshot);
-	}
 }
 #endif // ENABLE_VISUAL_LOG
 
diff --git a/Engine/Source/Runtime/GameplayTasks/Classes/GameplayTasksComponent.h b/Engine/Source/Runtime/GameplayTasks/Classes/GameplayTasksComponent.h
--- a/Engine/Source/Runtime/GameplayTasks/Classes/GameplayTasksComponent.h
+++ b/Engine/Source/Runtime/GameplayTasks/Classes/GameplayTasksComponent.h
@@ -12,6 +12,7 @@
 #include "GameplayTaskOwnerInterface.h"
 #include "GameplayTask.h"
 #include "GameplayTaskResource.h"
+#include "VisualLogger/VisualLoggerDebugSnapshotInterface.h"
 #include "GameplayTasksComponent.generated.h"
 
 class AActor;
@@ -57,7 +58,7 @@ DECLARE_DYNAMIC_MULTICAST_DELEGATE_TwoParams(FOnClaimedResourcesChangeSignature,
 *	The core ActorComponent for interfacing with the GameplayAbilities System
 */
 UCLASS(ClassGroup = GameplayTasks, hidecategories = (Object, LOD, Lighting, Transform, Sockets, TextureStreaming), editinlinenew, meta = (BlueprintSpawnableComponent))
-class GAMEPLAYTASKS_API UGameplayTasksComponent : public UActorComponent, public IGameplayTaskOwnerInterface
+class GAMEPLAYTASKS_API UGameplayTasksComponent : public UActorComponent, public IGameplayTaskOwnerInterface, public IVisualLoggerDebugSnapshotInterface
 {
 	GENERATED_BODY()
 
@@ -176,8 +177,10 @@ public:
 	GameplayTaskContainerType::TConstIterator GetTickingTaskIterator() const;
 	GameplayTaskContainerType::TConstIterator GetKnownTaskIterator() const;
 	GameplayTaskContainerType::TConstIterator GetPriorityQueueIterator() const;
+	GameplayTaskContainerType::TConstIterator GetSimulatedTaskIterator() const;
 
 #if ENABLE_VISUAL_LOG
+	virtual void GrabDebugSnapshot(FVisualLogEntry* Snapshot) const override;
 	void DescribeSelfToVisLog(struct FVisualLogEntry* Snapshot) const;
 #endif // ENABLE_VISUAL_LOG
 
diff --git a/Engine/Source/Runtime/GameplayTasks/Private/GameplayTasksComponent.cpp b/Engine/Source/Runtime/GameplayTasks/Private/GameplayTasksComponent.cpp
--- a/Engine/Source/Runtime/GameplayTasks/Private/GameplayTasksComponent.cpp
+++ b/Engine/Source/Runtime/GameplayTasks/Private/GameplayTasksComponent.cpp
@@ -258,27 +258,29 @@ void UGameplayTasksComponent::TickComponent(float DeltaTime, enum ELevelTick Tic
 	case 0:
 		break;
 	case 1:
-		if (TickingTasks[0])
 		{
-			TickingTasks[0]->TickTask(DeltaTime);
-			NumActuallyTicked++;
+			UGameplayTask* TickingTask = TickingTasks[0];
+			if (IsValid(TickingTask))
+			{
+				TickingTask->TickTask(DeltaTime);
+				NumActuallyTicked++;
+			}
 		}
 		break;
 	default:
-	{
-
-		static TArray<UGameplayTask*> LocalTickingTasks;
-		LocalTickingTasks.Reset();
-		LocalTickingTasks.Append(TickingTasks);
-		for (UGameplayTask* TickingTask : LocalTickingTasks)
 		{
-			if (TickingTask)
+			static TArray<UGameplayTask*> LocalTickingTasks;
+			LocalTickingTasks.Reset();
+			LocalTickingTasks.Append(TickingTasks);
+			for (UGameplayTask* TickingTask : LocalTickingTasks)
 			{
-				TickingTask->TickTask(DeltaTime);
-				NumActuallyTicked++;
+				if (IsValid(TickingTask))
+				{
+					TickingTask->TickTask(DeltaTime);
+					NumActuallyTicked++;
+				}
 			}
 		}
-	}
 		break;
 	};
 
@@ -667,7 +669,18 @@ FConstGameplayTaskIterator UGameplayTasksComponent::GetPriorityQueueIterator() c
 	return TaskPriorityQueue.CreateConstIterator();
 }
 
+FConstGameplayTaskIterator UGameplayTasksComponent::GetSimulatedTaskIterator() const
+{
+	return SimulatedTasks.CreateConstIterator();
+}
+
 #if ENABLE_VISUAL_LOG
+
+void UGameplayTasksComponent::GrabDebugSnapshot(FVisualLogEntry* Snapshot) const
+{
+	DescribeSelfToVisLog(Snapshot);
+}
+
 void UGameplayTasksComponent::DescribeSelfToVisLog(FVisualLogEntry* Snapshot) const
 {
 	static const FString CategoryName = TEXT("GameplayTasks");
