Iris - Register replication fragments for GameplayTasks.

#rnx

diff --git a/Engine/Source/Runtime/GameplayTasks/Classes/GameplayTask.h b/Engine/Source/Runtime/GameplayTasks/Classes/GameplayTask.h
--- a/Engine/Source/Runtime/GameplayTasks/Classes/GameplayTask.h
+++ b/Engine/Source/Runtime/GameplayTasks/Classes/GameplayTask.h
@@ -310,6 +310,11 @@ protected:
 	virtual void OnGameplayTaskDeactivated(UGameplayTask& Task) override;
 	// IGameplayTaskOwnerInterface END
 
+#if UE_WITH_IRIS
+	/** Register all replication fragments */
+	virtual void RegisterReplicationFragments(UE::Net::FFragmentRegistrationContext& Context, UE::Net::EFragmentRegistrationFlags RegistrationFlags) override;
+#endif // UE_WITH_IRIS
+
 private:
 	friend UGameplayTasksComponent;
 
diff --git a/Engine/Source/Runtime/GameplayTasks/GameplayTasks.Build.cs b/Engine/Source/Runtime/GameplayTasks/GameplayTasks.Build.cs
--- a/Engine/Source/Runtime/GameplayTasks/GameplayTasks.Build.cs
+++ b/Engine/Source/Runtime/GameplayTasks/GameplayTasks.Build.cs
@@ -25,6 +25,8 @@ namespace UnrealBuildTool.Rules
 				CircularlyReferencedDependentModules.Add("UnrealEd");
                 //PrivateDependencyModuleNames.Add("GameplayTagsEditor");
 			}
+
+			SetupIrisSupport(Target);
 		}
 	}
 }
diff --git a/Engine/Source/Runtime/GameplayTasks/Private/GameplayTask.cpp b/Engine/Source/Runtime/GameplayTasks/Private/GameplayTask.cpp
--- a/Engine/Source/Runtime/GameplayTasks/Private/GameplayTask.cpp
+++ b/Engine/Source/Runtime/GameplayTasks/Private/GameplayTask.cpp
@@ -7,6 +7,10 @@
 #include "GameplayTaskResource.h"
 #include "GameplayTasksComponent.h"
 
+#if UE_WITH_IRIS
+#include "Iris/ReplicationSystem/ReplicationFragmentUtil.h"
+#endif // UE_WITH_IRIS
+
 #include UE_INLINE_GENERATED_CPP_BY_NAME(GameplayTask)
 
 UGameplayTask::UGameplayTask(const FObjectInitializer& ObjectInitializer)
@@ -469,3 +473,10 @@ void UGameplayTask::OnGameplayTaskInitialized(UGameplayTask& Task)
 	ChildTask = &Task;
 }
 
+#if UE_WITH_IRIS
+void UGameplayTask::RegisterReplicationFragments(UE::Net::FFragmentRegistrationContext& Context, UE::Net::EFragmentRegistrationFlags RegistrationFlags)
+{
+	// Build descriptors and allocate PropertyReplicationFragments for this object.
+	UE::Net::FReplicationFragmentUtil::CreateAndRegisterFragmentsForObject(this, Context, RegistrationFlags);
+}
+#endif // UE_WITH_IRIS
