Handle 0.0f and negative time input values on GameplayTask_WaitDelay as finishing on the next tick.

inspired by #contrib 11079

diff --git a/Engine/Plugins/Runtime/GameplayAbilities/Source/GameplayAbilities/Private/Abilities/Tasks/AbilityTask_WaitDelay.cpp b/Engine/Plugins/Runtime/GameplayAbilities/Source/GameplayAbilities/Private/Abilities/Tasks/AbilityTask_WaitDelay.cpp
--- a/Engine/Plugins/Runtime/GameplayAbilities/Source/GameplayAbilities/Private/Abilities/Tasks/AbilityTask_WaitDelay.cpp
+++ b/Engine/Plugins/Runtime/GameplayAbilities/Source/GameplayAbilities/Private/Abilities/Tasks/AbilityTask_WaitDelay.cpp
@@ -28,9 +28,16 @@ void UAbilityTask_WaitDelay::Activate()
 	UWorld* World = GetWorld();
 	TimeStarted = World->GetTimeSeconds();
 
-	// Use a dummy timer handle as we don't need to store it for later but we don't need to look for something to clear
-	FTimerHandle TimerHandle;
-	World->GetTimerManager().SetTimer(TimerHandle, this, &UAbilityTask_WaitDelay::OnTimeFinish, Time, false);
+	if (Time <= 0.0f)
+	{
+		World->GetTimerManager().SetTimerForNextTick(this, &UAbilityTask_WaitDelay::OnTimeFinish);
+	}
+	else
+	{
+		// Use a dummy timer handle as we don't need to store it for later but we don't need to look for something to clear
+		FTimerHandle TimerHandle;
+		World->GetTimerManager().SetTimer(TimerHandle, this, &UAbilityTask_WaitDelay::OnTimeFinish, Time, false);
+	}
 }
 
 void UAbilityTask_WaitDelay::OnTimeFinish()
diff --git a/Engine/Source/Runtime/GameplayTasks/Private/Tasks/GameplayTask_WaitDelay.cpp b/Engine/Source/Runtime/GameplayTasks/Private/Tasks/GameplayTask_WaitDelay.cpp
--- a/Engine/Source/Runtime/GameplayTasks/Private/Tasks/GameplayTask_WaitDelay.cpp
+++ b/Engine/Source/Runtime/GameplayTasks/Private/Tasks/GameplayTask_WaitDelay.cpp
@@ -27,11 +27,6 @@ UGameplayTask_WaitDelay* UGameplayTask_WaitDelay::TaskWaitDelay(TScriptInterface
 
 UGameplayTask_WaitDelay* UGameplayTask_WaitDelay::TaskWaitDelay(IGameplayTaskOwnerInterface& InTaskOwner, float Time, const uint8 Priority)
 {
-	if (Time <= 0.f)
-	{
-		return nullptr;
-	}
-
 	UGameplayTask_WaitDelay* MyTask = NewTaskUninitialized<UGameplayTask_WaitDelay>();
 	if (MyTask)
 	{
@@ -46,9 +41,16 @@ void UGameplayTask_WaitDelay::Activate()
 	UWorld* World = GetWorld();
 	TimeStarted = World->GetTimeSeconds();
 
-	// Use a dummy timer handle as we don't need to store it for later but we don't need to look for something to clear
-	FTimerHandle TimerHandle;
-	World->GetTimerManager().SetTimer(TimerHandle, this, &UGameplayTask_WaitDelay::OnTimeFinish, (float)Time, false);
+	if (Time <= 0.0f)
+	{
+		World->GetTimerManager().SetTimerForNextTick(this, &UGameplayTask_WaitDelay::OnTimeFinish);
+	}
+	else
+	{
+		// Use a dummy timer handle as we don't need to store it for later but we don't need to look for something to clear
+		FTimerHandle TimerHandle;
+		World->GetTimerManager().SetTimer(TimerHandle, this, &UGameplayTask_WaitDelay::OnTimeFinish, (float)Time, false);
+	}
 }
 
 void UGameplayTask_WaitDelay::OnTimeFinish()
