Removing OnRep ptrs and handling differences with a locally stored list of FWeakObjectPtr so we don't crash when repping GC'd tasks.

diff --git a/Engine/Source/Runtime/GameplayTasks/Classes/GameplayTasksComponent.h b/Engine/Source/Runtime/GameplayTasks/Classes/GameplayTasksComponent.h
--- a/Engine/Source/Runtime/GameplayTasks/Classes/GameplayTasksComponent.h
+++ b/Engine/Source/Runtime/GameplayTasks/Classes/GameplayTasksComponent.h
@@ -118,7 +118,7 @@ public:
 	GAMEPLAYTASKS_API UGameplayTasksComponent(const FObjectInitializer& ObjectInitializer);
 	
 	UFUNCTION()
-	GAMEPLAYTASKS_API void OnRep_SimulatedTasks(const TArray<UGameplayTask*>& PreviousSimulatedTasks);
+	GAMEPLAYTASKS_API void OnRep_SimulatedTasks();
 
 	GAMEPLAYTASKS_API virtual void GetLifetimeReplicatedProps(TArray< FLifetimeProperty >& OutLifetimeProps) const override;
 	GAMEPLAYTASKS_API virtual bool ReplicateSubobjects(UActorChannel *Channel, FOutBunch* Bunch, FReplicationFlags* RepFlags) override;
@@ -207,9 +207,12 @@ private:
 	GAMEPLAYTASKS_API void SetSimulatedTasksNetDirty();
 
 	/** Tasks that run on simulated proxies */
-	UPROPERTY(ReplicatedUsing = OnRep_SimulatedTasks)
+	UPROPERTY(Transient, ReplicatedUsing = OnRep_SimulatedTasks)
 	TArray<TObjectPtr<UGameplayTask>> SimulatedTasks;
 
+	/** Previously replicated tasks that run on simulated proxies */
+	TArray<TWeakObjectPtr<UGameplayTask>> PreviousOnRepSimulatedTasks;
+
 };
 
 typedef UGameplayTasksComponent::GameplayTaskContainerType::TConstIterator FConstGameplayTaskIterator;
diff --git a/Engine/Source/Runtime/GameplayTasks/Private/GameplayTasksComponent.cpp b/Engine/Source/Runtime/GameplayTasks/Private/GameplayTasksComponent.cpp
--- a/Engine/Source/Runtime/GameplayTasks/Private/GameplayTasksComponent.cpp
+++ b/Engine/Source/Runtime/GameplayTasks/Private/GameplayTasksComponent.cpp
@@ -202,13 +202,14 @@ bool UGameplayTasksComponent::ReplicateSubobjects(UActorChannel* Channel, class
 	return WroteSomething;
 }
 
-void UGameplayTasksComponent::OnRep_SimulatedTasks(const TArray<UGameplayTask*>& PreviousSimulatedTasks)
+void UGameplayTasksComponent::OnRep_SimulatedTasks()
 {
 	if (IsUsingRegisteredSubObjectList())
 	{
 		// Find if any tasks got removed
-		for (UGameplayTask* OldSimulatedTask : PreviousSimulatedTasks)
+		for (TWeakObjectPtr<UGameplayTask> OldSimulatedTaskPtr : PreviousOnRepSimulatedTasks)
 		{
+			UGameplayTask* OldSimulatedTask = OldSimulatedTaskPtr.Get();
 			if (OldSimulatedTask)
 			{
 				const bool bIsRemoved = SimulatedTasks.Find(OldSimulatedTask) == INDEX_NONE;
@@ -241,7 +242,7 @@ void UGameplayTasksComponent::OnRep_SimulatedTasks(const TArray<UGameplayTask*>&
 			// See if it's a new task that needs to be registered
 			if (IsUsingRegisteredSubObjectList())
 			{
-				const bool bIsNew = PreviousSimulatedTasks.Find(SimulatedTask) == INDEX_NONE;
+				const bool bIsNew = PreviousOnRepSimulatedTasks.Find(SimulatedTask) == INDEX_NONE;
 				if (bIsNew)
 				{
 					AddReplicatedSubObject(SimulatedTask, COND_SkipOwner);
@@ -249,6 +250,9 @@ void UGameplayTasksComponent::OnRep_SimulatedTasks(const TArray<UGameplayTask*>&
 			}
 		}
 	}
+	
+	PreviousOnRepSimulatedTasks.Empty();
+	PreviousOnRepSimulatedTasks.Append(SimulatedTasks);
 }
 
 void UGameplayTasksComponent::TickComponent(float DeltaTime, enum ELevelTick TickType, FActorComponentTickFunction *ThisTickFunction)
