Make sure our gameplay task debug define is enabled if visual Logging is enabled so that compilation doesn't fail when force enabling debug for shipping and test builds.

#Synced-CL 49696721

diff --git a/Engine/Source/Runtime/GameplayTasks/Public/GameplayTaskTypes.h b/Engine/Source/Runtime/GameplayTasks/Public/GameplayTaskTypes.h
--- a/Engine/Source/Runtime/GameplayTasks/Public/GameplayTaskTypes.h
+++ b/Engine/Source/Runtime/GameplayTasks/Public/GameplayTaskTypes.h
@@ -3,12 +3,13 @@
 #pragma once
 
 #include "CoreMinimal.h"
+#include "VisualLogger/VisualLoggerDefines.h"
 #include "UObject/Object.h"
 #include "UObject/WeakObjectPtr.h"
 #include "Templates/Casts.h"
 
 #ifndef WITH_GAMEPLAYTASK_DEBUG
-#define WITH_GAMEPLAYTASK_DEBUG (WITH_UNREAL_DEVELOPER_TOOLS || !(UE_BUILD_SHIPPING || UE_BUILD_TEST))
+#define WITH_GAMEPLAYTASK_DEBUG (WITH_UNREAL_DEVELOPER_TOOLS || !(UE_BUILD_SHIPPING || UE_BUILD_TEST) || UE_DEBUG_RENDERING_ENABLED)
 #endif // WITH_GAMEPLAYTASK_DEBUG
 
 namespace FGameplayTasks
